Added include for Paths.h

diff --git a/Engine/Source/Runtime/PhysicsCore/Private/PhysicsInitialization.cpp b/Engine/Source/Runtime/PhysicsCore/Private/PhysicsInitialization.cpp
--- a/Engine/Source/Runtime/PhysicsCore/Private/PhysicsInitialization.cpp
+++ b/Engine/Source/Runtime/PhysicsCore/Private/PhysicsInitialization.cpp
@@ -8,6 +8,7 @@
 #include "IPhysXCookingModule.h"
 #include "IPhysXCooking.h"
 #include "Modules/ModuleManager.h"
+#include "Misc/Paths.h"
 
 #ifndef APEX_STATICALLY_LINKED
 #define APEX_STATICALLY_LINKED	0
