Copying //UE4/Dev-Physics to Dev-Main (//UE4/Dev-Main) [at] 7703071

diff --git a/Engine/Source/Runtime/Experimental/GeometryCollectionEngine/Private/GeometryCollection/PhysicsAssetSimulation.cpp b/Engine/Source/Runtime/Experimental/GeometryCollectionEngine/Private/GeometryCollection/PhysicsAssetSimulation.cpp
--- a/Engine/Source/Runtime/Experimental/GeometryCollectionEngine/Private/GeometryCollection/PhysicsAssetSimulation.cpp
+++ b/Engine/Source/Runtime/Experimental/GeometryCollectionEngine/Private/GeometryCollection/PhysicsAssetSimulation.cpp
@@ -3,7 +3,7 @@
 #include "GeometryCollection/GeometryCollectionCollisionStructureManager.h"
 #include "GeometryCollection/GeometryCollectionSimulationTypes.h"
 
-#include "PBDRigidsSolver.h"
+#include "PhysicsSolver.h"
 #include "ChaosSolversModule.h"
 #include "Chaos/DebugDrawQueue.h"
 #include "Chaos/ErrorReporter.h"
@@ -17,7 +17,7 @@
 #include "PhysicalMaterials/Experimental/ChaosPhysicalMaterial.h"
 #include "PhysicsEngine/PhysicsAsset.h"
 #include "Rendering/SkeletalMeshRenderData.h"
-#include "SolverObjects/SkeletalMeshPhysicsObject.h"
+#include "PhysicsProxy/SkeletalMeshPhysicsProxy.h"
 
 #include "PhysXIncludes.h"
 
@@ -30,7 +30,7 @@
 
 #if INCLUDE_CHAOS
 
-void FPhysicsAssetSimulationUtil::BuildParams(const UObject* Caller, const AActor* OwningActor, const USkeletalMeshComponent* SkelMeshComponent, const UPhysicsAsset* PhysicsAsset, FSkeletalMeshPhysicsObjectParams& Params)
+void FPhysicsAssetSimulationUtil::BuildParams(const UObject* Caller, const AActor* OwningActor, const USkeletalMeshComponent* SkelMeshComponent, const UPhysicsAsset* PhysicsAsset, FSkeletalMeshPhysicsProxyParams& Params)
 {
 #if WITH_PHYSX
 	Params.LocalToWorld = Params.InitialTransform = OwningActor->GetTransform();
@@ -317,7 +317,9 @@ void FPhysicsAssetSimulationUtil::BuildParams(const UObject* Caller, const AActo
 					const float LocalMaxExtent = Bounds.GetExtent().GetMax();
 					Bounds.ExpandBy(Extents.Size() / 10);
 
-					Chaos::TParticles<float, 3> Particles(Points);
+					//TODO: this gets used later which is incorrect. However, it was already broken because it passed a view and then moved the underlying data.
+					//Keeping as is for now, needs fixing
+					Chaos::TParticles<float, 3> Particles(MoveTemp(Points));
 					TUniquePtr<Chaos::TTriangleMesh<float>> TriangleMesh(
 						new Chaos::TTriangleMesh<float>(MoveTemp(Triangles)));
 					Chaos::FErrorReporter ErrorReporter;
@@ -455,7 +457,7 @@ void FPhysicsAssetSimulationUtil::BuildParams(const UObject* Caller, const AActo
 #endif
 }
 
-bool FPhysicsAssetSimulationUtil::UpdateAnimState(const UObject* Caller, const AActor* OwningActor, const USkeletalMeshComponent* SkelMeshComponent, const float Dt, FSkeletalMeshPhysicsObjectParams& Params)
+bool FPhysicsAssetSimulationUtil::UpdateAnimState(const UObject* Caller, const AActor* OwningActor, const USkeletalMeshComponent* SkelMeshComponent, const float Dt, FSkeletalMeshPhysicsProxyParams& Params)
 {
 	int32 UpdatedBones = 0;
 
@@ -481,10 +483,13 @@ bool FPhysicsAssetSimulationUtil::UpdateAnimState(const UObject* Caller, const A
 	Params.BoneHierarchy.SetActorWorldSpaceTransform(OwningActor->GetTransform());
 
 	// Update bone transforms from animation.
+	// @todo(ccaulfield): BoneSpaceTransforms is now copied - can we avoid that?
+	// The ocnst cast is unfortunate - maybe the task waiting in the component should use mutable...
+	TArray<FTransform> BoneSpaceTransforms = const_cast<USkeletalMeshComponent*>(SkelMeshComponent)->GetBoneSpaceTransforms();
 	const TArray<int32>& BoneIndices = Params.BoneHierarchy.GetBoneIndices();
 	for (const int32 BoneIndex : BoneIndices)
 	{
-		if (SkelMeshComponent->BoneSpaceTransforms.IsValidIndex(BoneIndex))
+		if (BoneSpaceTransforms.IsValidIndex(BoneIndex))
 		{
 			const int32 SocketIndex = Params.BoneHierarchy.GetSocketIndexForBone(BoneIndex);
 			if (SocketIndex != INDEX_NONE)
diff --git a/Engine/Source/Runtime/Experimental/GeometryCollectionEngine/Private/GeometryCollection/SkeletalMeshSimulationComponent.cpp b/Engine/Source/Runtime/Experimental/GeometryCollectionEngine/Private/GeometryCollection/SkeletalMeshSimulationComponent.cpp
--- a/Engine/Source/Runtime/Experimental/GeometryCollectionEngine/Private/GeometryCollection/SkeletalMeshSimulationComponent.cpp
+++ b/Engine/Source/Runtime/Experimental/GeometryCollectionEngine/Private/GeometryCollection/SkeletalMeshSimulationComponent.cpp
@@ -4,7 +4,7 @@
 #include "GeometryCollection/GeometryCollectionSimulationTypes.h"
 #include "GeometryCollection/PhysicsAssetSimulation.h"
 
-#include "PBDRigidsSolver.h"
+#include "PhysicsSolver.h"
 #include "ChaosSolversModule.h"
 #include "Chaos/DebugDrawQueue.h"
 #include "Chaos/ErrorReporter.h"
@@ -18,7 +18,7 @@
 #include "Engine/SkeletalMesh.h"
 #include "Engine/SkeletalMeshSocket.h"
 #include "Rendering/SkeletalMeshRenderData.h"
-#include "SolverObjects/SkeletalMeshPhysicsObject.h"
+#include "PhysicsProxy/SkeletalMeshPhysicsProxy.h"
 
 #include "PhysXIncludes.h"
 
@@ -64,7 +64,7 @@ USkeletalMeshSimulationComponent::USkeletalMeshSimulationComponent(const FObject
 	, InitialAngularVelocity(0.f)
 
 #if INCLUDE_CHAOS
-	, PhysicsObject(nullptr)
+	, PhysicsProxy(nullptr)
 #endif // INCLUDE_CHAOS
 {
 	// Enable calls to TickComponent()
@@ -75,7 +75,7 @@ USkeletalMeshSimulationComponent::USkeletalMeshSimulationComponent(const FObject
 }
 
 #if INCLUDE_CHAOS
-Chaos::FPBDRigidsSolver* GetSolver(const USkeletalMeshSimulationComponent& SkeletalMeshSimulationComponent)
+Chaos::FPhysicsSolver* GetSolver(const USkeletalMeshSimulationComponent& SkeletalMeshSimulationComponent)
 {
 	return	SkeletalMeshSimulationComponent.ChaosSolverActor != nullptr ?
 		SkeletalMeshSimulationComponent.ChaosSolverActor->GetSolver() :
@@ -113,7 +113,7 @@ void USkeletalMeshSimulationComponent::OnCreatePhysicsState()
 			ChaosMaterial->SleepingAngularThreshold = PhysicalMaterial->SleepingAngularVelocityThreshold;
 		}
 
-		auto InitFunc = [this, OwningActor, SkelMeshComponent](FSkeletalMeshPhysicsObjectParams& OutPhysicsParams)
+		auto InitFunc = [this, OwningActor, SkelMeshComponent](FSkeletalMeshPhysicsProxyParams& OutPhysicsParams)
 		{
 			OutPhysicsParams.bSimulating = bSimulating;
 
@@ -149,11 +149,11 @@ void USkeletalMeshSimulationComponent::OnCreatePhysicsState()
 			FPhysicsAssetSimulationUtil::UpdateAnimState(this, OwningActor, SkelMeshComponent, 0.0f, OutPhysicsParams);
 		};
 
-		check(PhysicsObject == nullptr);
-		PhysicsObject = new FSkeletalMeshPhysicsObject(this, InitFunc);
+		check(PhysicsProxy == nullptr);
+		PhysicsProxy = new FSkeletalMeshPhysicsProxy(this, InitFunc);
 
 		TSharedPtr<FPhysScene_Chaos> Scene = GetPhysicsScene();
-		Scene->AddObject(SkelMeshComponent, PhysicsObject);
+		Scene->AddObject(SkelMeshComponent, PhysicsProxy);
 
 		AChaosSolverActor* const SolverActor = Cast<AChaosSolverActor>(Scene->GetSolverActor());
 		UChaosGameplayEventDispatcher* const EventDispatcher = SolverActor ? SolverActor->GetGameplayEventDispatcher() : nullptr;
@@ -182,7 +182,7 @@ void USkeletalMeshSimulationComponent::OnDestroyPhysicsState()
 	UActorComponent::OnDestroyPhysicsState();
 
 #if INCLUDE_CHAOS
-	if (PhysicsObject)
+	if (PhysicsProxy)
 	{
 		// Remove our tick dependency on the Skeletal Mesh component.
 		AActor* OwningActor = GetOwner();
@@ -191,10 +191,10 @@ void USkeletalMeshSimulationComponent::OnDestroyPhysicsState()
 
 		// Handle scene remove, right now we rely on the reset of EndPlay to clean up
 		TSharedPtr<FPhysScene_Chaos> Scene = GetPhysicsScene();
-		Scene->RemoveObject(PhysicsObject);
+		Scene->RemoveObject(PhysicsProxy);
 
 		// Discard the pointer, the scene will handle destroying it
-		PhysicsObject = nullptr;
+		PhysicsProxy = nullptr;
 	}
 #endif // INCLUDE_CHAOS
 }
@@ -207,7 +207,7 @@ bool USkeletalMeshSimulationComponent::ShouldCreatePhysicsState() const
 bool USkeletalMeshSimulationComponent::HasValidPhysicsState() const
 {
 #if INCLUDE_CHAOS
-	return PhysicsObject != nullptr;
+	return PhysicsProxy != nullptr;
 #else // INCLUDE_CHAOS
 	return false;
 #endif // INCLUDE_CHAOS
@@ -254,8 +254,8 @@ void USkeletalMeshSimulationComponent::TickComponent(float DeltaTime, enum ELeve
 		{
 			AActor* OwningActor = GetOwner();
 			USkeletalMeshComponent* SkelMeshComponent = OwningActor->FindComponentByClass<USkeletalMeshComponent>();
-			PhysicsObject->CaptureInputs(DeltaTime,
-				[this, OwningActor, SkelMeshComponent](const float Dt, FSkeletalMeshPhysicsObjectParams & InOutPhysicsParams) -> bool
+			PhysicsProxy->CaptureInputs(DeltaTime,
+				[this, OwningActor, SkelMeshComponent](const float Dt, FSkeletalMeshPhysicsProxyParams & InOutPhysicsParams) -> bool
 				{
 					return FPhysicsAssetSimulationUtil::UpdateAnimState(this, OwningActor, SkelMeshComponent, Dt, InOutPhysicsParams);
 				});
diff --git a/Engine/Source/Runtime/Experimental/GeometryCollectionEngine/Private/GeometryCollection/StaticMeshSimulationComponent.cpp b/Engine/Source/Runtime/Experimental/GeometryCollectionEngine/Private/GeometryCollection/StaticMeshSimulationComponent.cpp
--- a/Engine/Source/Runtime/Experimental/GeometryCollectionEngine/Private/GeometryCollection/StaticMeshSimulationComponent.cpp
+++ b/Engine/Source/Runtime/Experimental/GeometryCollectionEngine/Private/GeometryCollection/StaticMeshSimulationComponent.cpp
@@ -14,8 +14,8 @@
 #include "ChaosSolversModule.h"
 #include "Components/StaticMeshComponent.h"
 #include "Engine/StaticMesh.h"
-#include "SolverObjects/StaticMeshPhysicsObject.h"
-#include "PBDRigidsSolver.h"
+#include "PhysicsProxy/StaticMeshPhysicsProxy.h"
+#include "PhysicsSolver.h"
 #include "Chaos/ChaosGameplayEventDispatcher.h"
 
 
@@ -57,16 +57,16 @@ void UStaticMeshSimulationComponent::TickComponent(float DeltaTime, enum ELevelT
 		Chaos::IDispatcher* PhysicsDispatcher = ChaosModule->GetDispatcher();
 		checkSlow(PhysicsDispatcher); // Should always have one of these
 
-		for (int32 Idx = 0; Idx < PhysicsObjects.Num(); ++Idx)
+		for (int32 Idx = 0; Idx < PhysicsProxies.Num(); ++Idx)
 		{
-			FStaticMeshPhysicsObject* const PhysicsObject = PhysicsObjects[Idx];
+			FStaticMeshPhysicsProxy* const PhysicsProxy = PhysicsProxies[Idx];
 			UPrimitiveComponent* const Comp = SimulatedComponents[Idx];
 
-			FSolverObjectKinematicUpdate ParamUpdate;
+			FPhysicsProxyKinematicUpdate ParamUpdate;
 			ParamUpdate.NewTransform = Comp->GetComponentTransform();
 			ParamUpdate.NewVelocity = Comp->ComponentVelocity;
 
-			PhysicsDispatcher->EnqueueCommand([PhysObj = PhysicsObject, Params = ParamUpdate]()
+			PhysicsDispatcher->EnqueueCommandImmediate([PhysObj = PhysicsProxy, Params = ParamUpdate]()
 			{
 				PhysObj->BufferKinematicUpdate(Params);
 			});
@@ -76,7 +76,7 @@ void UStaticMeshSimulationComponent::TickComponent(float DeltaTime, enum ELevelT
 }
 
 #if INCLUDE_CHAOS
-Chaos::FPBDRigidsSolver* GetSolver(const UStaticMeshSimulationComponent& StaticMeshSimulationComponent)
+Chaos::FPhysicsSolver* GetSolver(const UStaticMeshSimulationComponent& StaticMeshSimulationComponent)
 {
 	return	StaticMeshSimulationComponent.ChaosSolverActor != nullptr ? StaticMeshSimulationComponent.ChaosSolverActor->GetSolver() : StaticMeshSimulationComponent.GetOwner()->GetWorld()->PhysicsScene_Chaos->GetSolver();
 }
@@ -122,7 +122,7 @@ void UStaticMeshSimulationComponent::OnCreatePhysicsState()
 					ChaosMaterial->SleepingLinearThreshold = PhysicalMaterial->SleepingLinearVelocityThreshold;
 					ChaosMaterial->SleepingAngularThreshold = PhysicalMaterial->SleepingAngularVelocityThreshold;
 				}
-				auto InitFunc = [this, TargetComponent](FStaticMeshPhysicsObject::Params& InParams)
+				auto InitFunc = [this, TargetComponent](FStaticMeshPhysicsProxy::Params& InParams)
 				{
 					GetPathName(this, InParams.Name);
 					InParams.InitialTransform = GetOwner()->GetTransform();
@@ -148,14 +148,15 @@ void UStaticMeshSimulationComponent::OnCreatePhysicsState()
 						{
 							if ((InParams.InitialTransform.GetScale3D() - FVector(1.f, 1.f, 1.f)).SizeSquared() < SMALL_NUMBER)
 							{
-								InParams.MeshVertexPositions = MoveTemp(CollisionData.Vertices);
+								auto& TVectorArray = reinterpret_cast<TArray<Chaos::TVector<float, 3>>&>(CollisionData.Vertices);
+								InParams.MeshVertexPositions = MoveTemp(TVectorArray);
 							}
 							else
 							{
-								InParams.MeshVertexPositions.SetNum(CollisionData.Vertices.Num());
+								InParams.MeshVertexPositions.Resize(CollisionData.Vertices.Num());
 								for (int32 i = 0; i < CollisionData.Vertices.Num(); ++i)
 								{
-									InParams.MeshVertexPositions[i] = CollisionData.Vertices[i] * InParams.InitialTransform.GetScale3D();
+									InParams.MeshVertexPositions.X(i) = CollisionData.Vertices[i] * InParams.InitialTransform.GetScale3D();
 								}
 							}
 							check(sizeof(Chaos::TVector<int32, 3>) == sizeof(FTriIndices)); // binary compatible?
@@ -237,12 +238,12 @@ void UStaticMeshSimulationComponent::OnCreatePhysicsState()
 					TargetComponent->SetWorldTransform(InTransform);
 				};
 
-				FStaticMeshPhysicsObject* const NewPhysicsObject = new FStaticMeshPhysicsObject(this, InitFunc, SyncFunc);
-				PhysicsObjects.Add(NewPhysicsObject);
+				FStaticMeshPhysicsProxy* const NewPhysicsProxy = new FStaticMeshPhysicsProxy(this, InitFunc, SyncFunc);
+				PhysicsProxies.Add(NewPhysicsProxy);
 				SimulatedComponents.Add(TargetComponent);
-				check(PhysicsObjects.Num() == SimulatedComponents.Num());
+				check(PhysicsProxies.Num() == SimulatedComponents.Num());
 
-				Scene->AddObject(TargetComponent, NewPhysicsObject);
+				Scene->AddObject(TargetComponent, NewPhysicsProxy);
 
 				if (EventDispatcher)
 				{
@@ -276,7 +277,7 @@ void UStaticMeshSimulationComponent::OnCreatePhysicsState()
 					ChaosMaterial->SleepingAngularThreshold = PhysicalMaterial->SleepingAngularVelocityThreshold;
 				}
 				const TArray<UStaticMeshComponent*>* StaticMeshComponentChildren = ParentToChildMap.Find(TargetComponent);
-				auto InitFunc = [this, StaticMeshComponentChildren, TargetComponent](FStaticMeshPhysicsObject::Params& InParams)
+				auto InitFunc = [this, StaticMeshComponentChildren, TargetComponent](FStaticMeshPhysicsProxy::Params& InParams)
 				{
 					GetPathName(this, InParams.Name);
 					InParams.InitialTransform = GetOwner()->GetTransform();
@@ -360,12 +361,12 @@ void UStaticMeshSimulationComponent::OnCreatePhysicsState()
 					TargetComponent->SetWorldTransform(InTransform);
 				};
 
-				FStaticMeshPhysicsObject* const NewPhysicsObject = new FStaticMeshPhysicsObject(this, InitFunc, SyncFunc);
-				PhysicsObjects.Add(NewPhysicsObject);
+				FStaticMeshPhysicsProxy* const NewPhysicsProxy = new FStaticMeshPhysicsProxy(this, InitFunc, SyncFunc);
+				PhysicsProxies.Add(NewPhysicsProxy);
 				SimulatedComponents.Add(TargetComponent);
-				check(PhysicsObjects.Num() == SimulatedComponents.Num());
+				check(PhysicsProxies.Num() == SimulatedComponents.Num());
 
-				Scene->AddObject(TargetComponent, NewPhysicsObject);
+				Scene->AddObject(TargetComponent, NewPhysicsProxy);
 
 				if (EventDispatcher)
 				{
@@ -400,16 +401,16 @@ void UStaticMeshSimulationComponent::OnDestroyPhysicsState()
 
 	if (Scene)
 	{
-		for (FStaticMeshPhysicsObject* PhysicsObject : PhysicsObjects)
+		for (FStaticMeshPhysicsProxy* PhysicsProxy : PhysicsProxies)
 		{
-			if (PhysicsObject)
+			if (PhysicsProxy)
 			{
 				// Handle scene remove, right now we rely on the reset of EndPlay to clean up
-				Scene->RemoveObject(PhysicsObject);
+				Scene->RemoveObject(PhysicsProxy);
 
 				if (EventDispatcher)
 				{
-					if (UPrimitiveComponent* const Comp = Scene->GetOwningComponent<UPrimitiveComponent>(PhysicsObject))
+					if (UPrimitiveComponent* const Comp = Scene->GetOwningComponent<UPrimitiveComponent>(PhysicsProxy))
 					{
 						EventDispatcher->UnRegisterForCollisionEvents(Comp, this);
 						EventDispatcher->UnRegisterForCollisionEvents(Comp, Comp);
@@ -419,7 +420,7 @@ void UStaticMeshSimulationComponent::OnDestroyPhysicsState()
 		}
 	}
 
-	PhysicsObjects.Empty();
+	PhysicsProxies.Empty();
 	SimulatedComponents.Empty();
 #endif
 }
@@ -431,7 +432,7 @@ bool UStaticMeshSimulationComponent::ShouldCreatePhysicsState() const
 
 bool UStaticMeshSimulationComponent::HasValidPhysicsState() const
 {
-	return PhysicsObjects.Num() > 0;
+	return PhysicsProxies.Num() > 0;
 }
 
 #if INCLUDE_CHAOS
diff --git a/Engine/Source/Runtime/Experimental/GeometryCollectionEngine/Public/ChaosBlueprint.h b/Engine/Source/Runtime/Experimental/GeometryCollectionEngine/Public/ChaosBlueprint.h
--- a/Engine/Source/Runtime/Experimental/GeometryCollectionEngine/Public/ChaosBlueprint.h
+++ b/Engine/Source/Runtime/Experimental/GeometryCollectionEngine/Public/ChaosBlueprint.h
@@ -3,9 +3,11 @@
 
 #include "CoreMinimal.h"
 #include "UObject/ObjectMacros.h"
-#include "PBDRigidsSolver.h"
+#include "Chaos/Declares.h"
 #include "Chaos/ChaosSolver.h"
 #include "Chaos/ChaosSolverActor.h"
+#include "EventManager.h"
+#include "EventsData.h"
 #include "HAL/ThreadSafeBool.h"
 #include "GeometryCollection/GeometryCollectionActor.h"
 #include "ChaosCollisionEventFilter.h"
@@ -36,7 +38,7 @@ UCLASS(ClassGroup = (Chaos, Common), hidecategories = (Object, ActorComponent, P
 class GEOMETRYCOLLECTIONENGINE_API UChaosDestructionListener : public USceneComponent
 {
 	GENERATED_UCLASS_BODY()
-
+		
 public:
 	//~ Begin UObject Interface
 #if WITH_EDITOR
@@ -45,6 +47,8 @@ public:
 	//~ End UObject Interface
 	
 	//~ Begin UActorComponent interface
+	virtual void BeginPlay() override;
+	virtual void EndPlay(const EEndPlayReason::Type EndPlayReason) override;
 	virtual void TickComponent(float DeltaTime, enum ELevelTick TickType, FActorComponentTickFunction *ThisTickFunction) override;
 	//~ End UActorComponent interface
 
@@ -149,18 +153,14 @@ public:
 	void SortTrailingEvents(UPARAM(ref) TArray<FChaosTrailingEventData>& TrailingEvents, EChaosTrailingSortMethod SortMethod);
 
 private:
-
 	// Updates the scene component transform settings
 	void UpdateTransformSettings();
 
 	// Retrieves data from solvers
 	void GetDataFromSolvers();
 
-	// Retrieve data from GeometryCollectionPhysicsObjects
-	void GetDataFromGeometryCollectionPhysicsObjects();
-
-	void UpdateSolvers();
-	void UpdateGeometryCollectionPhysicsObjects();
+	void ClearEvents();
+	void UpdateEvents();
 
 protected:
 
@@ -180,9 +180,9 @@ protected:
 
 #if INCLUDE_CHAOS
 	// The raw data arrays derived from the solvers
-	Chaos::FPBDRigidsSolver::FCollisionDataArray RawCollisionDataArray;
-	Chaos::FPBDRigidsSolver::FBreakingDataArray RawBreakingDataArray;
-	TArray<Chaos::TTrailingData<float, 3>> RawTrailingDataArray;
+	Chaos::FCollisionDataArray RawCollisionDataArray;
+	Chaos::FBreakingDataArray RawBreakingDataArray;
+	Chaos::FTrailingDataArray RawTrailingDataArray;
 #endif
 
 	FTransform ChaosComponentTransform;
@@ -195,10 +195,10 @@ protected:
 
 #if INCLUDE_CHAOS
 	// The list of rigid body solvers, used to retrieve destruction events
-	TSet<Chaos::FPBDRigidsSolver*> Solvers;
+	TSet<Chaos::FPhysicsSolver*> Solvers;
 
-	// The list of GeometryCollectionPhysicsObjects, used to retrieve destruction events
-	TArray<const FGeometryCollectionPhysicsObject*> GeometryCollectionPhysicsObjects;
+	// The list of GeometryCollectionPhysicsProxy, used to retrieve destruction events
+	TArray<const FGeometryCollectionPhysicsProxy*> GeometryCollectionPhysicsProxies;
 
 	TSharedPtr<FChaosCollisionEventFilter> ChaosCollisionFilter;
 	TSharedPtr<FChaosBreakingEventFilter> ChaosBreakingFilter;
@@ -210,14 +210,26 @@ public:
 		ChaosCollisionFilter = InCollisionFilter;
 	}
 
-	void SetBreakingFilter(TSharedPtr < FChaosBreakingEventFilter> InBreakingFilter)
+	void SetBreakingFilter(TSharedPtr<FChaosBreakingEventFilter> InBreakingFilter)
 	{
 		ChaosBreakingFilter = InBreakingFilter;
 	}
 
-	void SetTrailingFilter(TSharedPtr < FChaosTrailingEventFilter> InTrailingFilter)
+	void SetTrailingFilter(TSharedPtr<FChaosTrailingEventFilter> InTrailingFilter)
 	{
 		ChaosTrailingFilter = InTrailingFilter;
 	}
-#endif
+
+	void RegisterChaosEvents(FPhysScene* Scene);
+	void UnregisterChaosEvents(FPhysScene* Scene);
+
+	void RegisterChaosEvents(TSharedPtr<FPhysScene_Chaos> Scene);
+	void UnregisterChaosEvents(TSharedPtr<FPhysScene_Chaos> Scene);
+
+
+	// Chaos Event Handlers
+	void HandleCollisionEvents(const Chaos::FCollisionEventData& CollisionData);
+	void HandleBreakingEvents(const Chaos::FBreakingEventData& BreakingData);
+	void HandleTrailingEvents(const Chaos::FTrailingEventData& TrailingData);
+#endif // INCLUDE_CHAOS
 };
diff --git a/Engine/Source/Runtime/Experimental/GeometryCollectionEngine/Public/ChaosBreakingEventFilter.h b/Engine/Source/Runtime/Experimental/GeometryCollectionEngine/Public/ChaosBreakingEventFilter.h
--- a/Engine/Source/Runtime/Experimental/GeometryCollectionEngine/Public/ChaosBreakingEventFilter.h
+++ b/Engine/Source/Runtime/Experimental/GeometryCollectionEngine/Public/ChaosBreakingEventFilter.h
@@ -3,9 +3,9 @@
 
 #include "CoreMinimal.h"
 #include "UObject/ObjectMacros.h"
-#include "PBDRigidsSolver.h"
 #include "Chaos/ChaosSolver.h"
 #include "ChaosFilter.h"
+#include "EventsData.h"
 #include "ChaosBreakingEventFilter.generated.h"
 
 using namespace Chaos;
@@ -95,12 +95,12 @@ struct FChaosBreakingEventRequestSettings
 
 #if INCLUDE_CHAOS
 
-class GEOMETRYCOLLECTIONENGINE_API FChaosBreakingEventFilter : public IChaosEventFilter<FPBDRigidsSolver::FBreakingDataArray, TArray<FChaosBreakingEventData>, EChaosBreakingSortMethod>
+class GEOMETRYCOLLECTIONENGINE_API FChaosBreakingEventFilter : public IChaosEventFilter<Chaos::FBreakingDataArray, TArray<FChaosBreakingEventData>, EChaosBreakingSortMethod>
 {
 public:
 	FChaosBreakingEventFilter(FChaosBreakingEventRequestSettings* FilterSettingsIn) : BreakingEventRequestSettings(FilterSettingsIn) {}
 
-	virtual void FilterEvents(const FTransform& ChaosComponentTransform, const FPBDRigidsSolver::FBreakingDataArray& RawBreakingDataArray) override;
+	virtual void FilterEvents(const FTransform& ChaosComponentTransform, const Chaos::FBreakingDataArray& RawBreakingDataArray) override;
 
 	virtual void SortEvents(TArray<FChaosBreakingEventData>& InOutBreakingEvents, EChaosBreakingSortMethod SortMethod, const FTransform& InTransform) override;
 
diff --git a/Engine/Source/Runtime/Experimental/GeometryCollectionEngine/Public/ChaosCollisionEventFilter.h b/Engine/Source/Runtime/Experimental/GeometryCollectionEngine/Public/ChaosCollisionEventFilter.h
--- a/Engine/Source/Runtime/Experimental/GeometryCollectionEngine/Public/ChaosCollisionEventFilter.h
+++ b/Engine/Source/Runtime/Experimental/GeometryCollectionEngine/Public/ChaosCollisionEventFilter.h
@@ -3,9 +3,9 @@
 
 #include "CoreMinimal.h"
 #include "UObject/ObjectMacros.h"
-#include "PBDRigidsSolver.h"
 #include "Chaos/ChaosSolver.h"
 #include "ChaosFilter.h"
+#include "EventsData.h"
 #include "ChaosCollisionEventFilter.generated.h"
 
 using namespace Chaos;
@@ -44,13 +44,13 @@ struct FChaosCollisionEventData
 	UPROPERTY(EditAnywhere, BlueprintReadOnly, Category = "Data")
 	FVector Impulse;
 
-	// The particle index of the collision event
-	UPROPERTY(EditAnywhere, BlueprintReadOnly, Category = "Data")
-	int32 ParticleIndex;
+	// The particle of the collision event
+	//UPROPERTY(EditAnywhere, BlueprintReadOnly, Category = "Data")
+	Chaos::TGeometryParticle<float, 3>* Particle;
 
-	// The level set index of the collision event
-	UPROPERTY(EditAnywhere, BlueprintReadOnly, Category = "Data")
-	int32 LevelsetIndex;
+	// The level set of the collision event
+	//UPROPERTY(EditAnywhere, BlueprintReadOnly, Category = "Data")
+	Chaos::TGeometryParticle<float, 3>* Levelset;
 
 	FChaosCollisionEventData()
 		: Location(FVector::ZeroVector)
@@ -60,8 +60,8 @@ struct FChaosCollisionEventData
 		, Mass1(0.0f)
 		, Mass2(0.0f)
 		, Impulse(FVector::ZeroVector)
-		, ParticleIndex(INDEX_NONE)
-		, LevelsetIndex(INDEX_NONE)
+		, Particle(nullptr)
+		, Levelset(nullptr)
 	{
 	}
 };
@@ -121,13 +121,13 @@ struct FChaosCollisionEventRequestSettings
 
 #if INCLUDE_CHAOS
 
-class GEOMETRYCOLLECTIONENGINE_API FChaosCollisionEventFilter : public IChaosEventFilter<FPBDRigidsSolver::FCollisionDataArray, TArray<FChaosCollisionEventData>, EChaosCollisionSortMethod>
+class GEOMETRYCOLLECTIONENGINE_API FChaosCollisionEventFilter : public IChaosEventFilter<Chaos::FCollisionDataArray, TArray<FChaosCollisionEventData>, EChaosCollisionSortMethod>
  {
 public:
 
 	FChaosCollisionEventFilter(FChaosCollisionEventRequestSettings* FilterSettingsIn) : CollisionEventRequestSettings(FilterSettingsIn) {}
 
-	virtual void FilterEvents(const FTransform& ChaosComponentTransform, const FPBDRigidsSolver::FCollisionDataArray& RawCollisionDataArray) override;
+	virtual void FilterEvents(const FTransform& ChaosComponentTransform, const Chaos::FCollisionDataArray& RawCollisionDataArray) override;
 	virtual void SortEvents(TArray<FChaosCollisionEventData>& InOutCollisionEvents, EChaosCollisionSortMethod SortMethod, const FTransform& InTransform) override;
 
 private:
diff --git a/Engine/Source/Runtime/Experimental/GeometryCollectionEngine/Public/ChaosTrailingEventFilter.h b/Engine/Source/Runtime/Experimental/GeometryCollectionEngine/Public/ChaosTrailingEventFilter.h
--- a/Engine/Source/Runtime/Experimental/GeometryCollectionEngine/Public/ChaosTrailingEventFilter.h
+++ b/Engine/Source/Runtime/Experimental/GeometryCollectionEngine/Public/ChaosTrailingEventFilter.h
@@ -3,9 +3,9 @@
 
 #include "CoreMinimal.h"
 #include "UObject/ObjectMacros.h"
-#include "PBDRigidsSolver.h"
 #include "Chaos/ChaosSolver.h"
 #include "ChaosFilter.h"
+#include "EventsData.h"
 #include "ChaosTrailingEventFilter.generated.h"
 
 using namespace Chaos;
@@ -98,13 +98,14 @@ struct FChaosTrailingEventRequestSettings
 	}
 };
 
+#if INCLUDE_CHAOS
 class GEOMETRYCOLLECTIONENGINE_API FChaosTrailingEventFilter 
-	: public IChaosEventFilter<TArray<Chaos::TTrailingData<float, 3>>, TArray<FChaosTrailingEventData>, EChaosTrailingSortMethod>
+	: public IChaosEventFilter<Chaos::FTrailingDataArray, TArray<FChaosTrailingEventData>, EChaosTrailingSortMethod>
 {
 public:
 	FChaosTrailingEventFilter(FChaosTrailingEventRequestSettings* FilterSettingsIn) : TrailingEventRequestSettings(FilterSettingsIn) {}
 
-	virtual void FilterEvents(const FTransform& ChaosComponentTransform, const TArray<Chaos::TTrailingData<float, 3>>& RawTrailingDataArray) override;
+	virtual void FilterEvents(const FTransform& ChaosComponentTransform, const Chaos::FTrailingDataArray& RawTrailingDataArray) override;
 
 	virtual void SortEvents(TArray<FChaosTrailingEventData>& InOutTrailingEvents, EChaosTrailingSortMethod SortMethod, const FTransform& InTransform) override;
 
@@ -112,3 +113,4 @@ private:
 	FChaosTrailingEventFilter() : TrailingEventRequestSettings(nullptr) {}
 	const FChaosTrailingEventRequestSettings* TrailingEventRequestSettings;
 };
+#endif
diff --git a/Engine/Source/Runtime/Experimental/GeometryCollectionEngine/Public/GeometryCollection/GeometryCollectionActor.h b/Engine/Source/Runtime/Experimental/GeometryCollectionEngine/Public/GeometryCollection/GeometryCollectionActor.h
--- a/Engine/Source/Runtime/Experimental/GeometryCollectionEngine/Public/GeometryCollection/GeometryCollectionActor.h
+++ b/Engine/Source/Runtime/Experimental/GeometryCollectionEngine/Public/GeometryCollection/GeometryCollectionActor.h
@@ -10,7 +10,6 @@
 #include "GeometryCollection/GeometryCollectionComponent.h"
 #include "GeometryCollection/GeometryCollectionDebugDrawComponent.h"
 #include "Physics/Experimental/PhysScene_Chaos.h"
-#include "Physics/Experimental/PhysScene_LLImmediate.h"
 #include "Physics/ImmediatePhysics/ImmediatePhysicsActorHandle.h"
 #include "UObject/ObjectMacros.h"
 
@@ -45,30 +44,34 @@ public:
 	bool RaycastSingle(FVector Start, FVector End, FHitResult& OutHit) const;
 
 
-#if ! INCLUDE_CHAOS
-	/// Stub support for Immediate mode. 
-public:
-
-	void StartFrameCallback(float EndFrame);
-
-	void CreateRigidBodyCallback(FSolverCallbacks::FParticlesType& Particles);
-
-	void EndFrameCallback(float EndFrame);
-
-	void UpdateKinematicBodiesCallback(FSolverCallbacks::FParticlesType& Particles, const float Dt, const float Time, const int32 Index);
-
-	void ParameterUpdateCallback(FSolverCallbacks::FParticlesType& Particles, const float Time);
-
-	void DisableCollisionsCallback(TSet<TTuple<int32, int32>>& CollisionPairs);
-
-	void AddForceCallback(FSolverCallbacks::FParticlesType& Particles, const float Dt, const int32 Index);
-
-private:
-	FPhysScene_LLImmediate Scene;
-	bool bInitializedState;
-	TSharedRef< TManagedArray<int32> > RigidBodyIdArray;
-	TSharedRef< TManagedArray<FVector> > CenterOfMassArray;
-#endif
-
-
-};
\ No newline at end of file
+// #if ! INCLUDE_CHAOS
+// 	/// Stub support for Immediate mode. 
+// public:
+// 
+// 	void StartFrameCallback(float EndFrame);
+// 
+// 	void CreateRigidBodyCallback(FSolverCallbacks::FParticlesType& Particles);
+// 
+// 	void EndFrameCallback(float EndFrame);
+// 
+// 	void UpdateKinematicBodiesCallback(FSolverCallbacks::FParticlesType& Particles, const float Dt, const float Time, const int32 Index);
+// 
+// 	void ParameterUpdateCallback(FSolverCallbacks::FParticlesType& Particles, const float Time);
+// 
+// 	void DisableCollisionsCallback(TSet<TTuple<int32, int32>>& CollisionPairs);
+// 
+// 	void AddForceCallback(FSolverCallbacks::FParticlesType& Particles, const float Dt, const int32 Index);
+// 
+// private:
+// 	FPhysScene_LLImmediate Scene;
+// 	bool bInitializedState;
+// 	TSharedRef< TManagedArray<int32> > RigidBodyIdArray;
+// 	TSharedRef< TManagedArray<FVector> > CenterOfMassArray;
+// #endif
+};
+
+#if !INCLUDE_CHAOS
+inline AGeometryCollectionActor::AGeometryCollectionActor(const FObjectInitializer& ObjectInitializer) : Super(ObjectInitializer) {}
+inline void AGeometryCollectionActor::Tick(float DeltaSeconds) {}
+inline bool AGeometryCollectionActor::RaycastSingle(FVector Start, FVector End, FHitResult& OutHit) const { return false; }
+#endif
\ No newline at end of file
diff --git a/Engine/Source/Runtime/Experimental/GeometryCollectionEngine/Public/GeometryCollection/GeometryCollectionComponent.h b/Engine/Source/Runtime/Experimental/GeometryCollectionEngine/Public/GeometryCollection/GeometryCollectionComponent.h
--- a/Engine/Source/Runtime/Experimental/GeometryCollectionEngine/Public/GeometryCollection/GeometryCollectionComponent.h
+++ b/Engine/Source/Runtime/Experimental/GeometryCollectionEngine/Public/GeometryCollection/GeometryCollectionComponent.h
@@ -15,7 +15,6 @@
 #include "GeometryCollection/GeometryCollectionSimulationTypes.h"
 #include "GeometryCollection/GeometryCollectionSimulationCoreTypes.h"
 #include "Physics/Experimental/PhysScene_Chaos.h"
-#include "Physics/Experimental/PhysScene_LLImmediate.h"
 #include "GeometryCollection/GeometryDynamicCollection.h"
 #include "GeometryCollectionObject.h"
 #include "GeometryCollection/RecordedTransformTrack.h"
@@ -33,7 +32,7 @@
 
 struct FGeometryCollectionConstantData;
 struct FGeometryCollectionDynamicData;
-class FGeometryCollectionPhysicsObject;
+class FGeometryCollectionPhysicsProxy;
 class UGeometryCollectionComponent;
 class UBoxComponent;
 class UGeometryCollectionCache;
@@ -495,7 +494,7 @@ public:
 #if INCLUDE_CHAOS
 	const TSharedPtr<FPhysScene_Chaos> GetPhysicsScene() const;
 	AChaosSolverActor* GetPhysicsSolverActor() const;
-	const FGeometryCollectionPhysicsObject* GetPhysicsObject() const { return PhysicsObject; }
+	const FGeometryCollectionPhysicsProxy* GetPhysicsProxy() const { return PhysicsProxy; }
 #endif  // #if INCLUDE_CHAOS
 
 #if GEOMETRYCOLLECTION_EDITOR_SELECTION
@@ -615,7 +614,7 @@ private:
 	//@todo(mlentine): Don't have one per geo collection
 	TUniquePtr<Chaos::TChaosPhysicsMaterial<float>> ChaosMaterial;
 
-	FGeometryCollectionPhysicsObject* PhysicsObject;
+	FGeometryCollectionPhysicsProxy* PhysicsProxy;
 	TUniquePtr<FGeometryDynamicCollection> DynamicCollection;
 	TArray<FManagedArrayBase**> CopyOnWriteAttributeList;
 
diff --git a/Engine/Source/Runtime/Experimental/GeometryCollectionEngine/Public/GeometryCollection/GeometryCollectionParticlesData.h b/Engine/Source/Runtime/Experimental/GeometryCollectionEngine/Public/GeometryCollection/GeometryCollectionParticlesData.h
--- a/Engine/Source/Runtime/Experimental/GeometryCollectionEngine/Public/GeometryCollection/GeometryCollectionParticlesData.h
+++ b/Engine/Source/Runtime/Experimental/GeometryCollectionEngine/Public/GeometryCollection/GeometryCollectionParticlesData.h
@@ -12,9 +12,9 @@
 #include "Chaos/UniformGrid.h"
 #include "Chaos/PBDRigidClustering.h"
 #include "Chaos/Framework/BufferedData.h"
+#include "Chaos/Declares.h"
 
 class FChaosSolversModule;
-namespace Chaos { class FPBDRigidsSolver; }
 template<class InElementType> class TManagedArray;
 
 /** Enumeration of the synchronizable data. */
@@ -78,7 +78,7 @@ public:
 	bool RequestSyncedData(EGeometryCollectionParticlesData Data) const { SetDataSyncFlag(Data); return HasSyncedData(Data); }
 
 	/** Copy the data of the specified set of particles/rigid body ids to this object. */
-	void Sync(const Chaos::FPBDRigidsSolver* Solver, const TManagedArray<int32>& RigidBodyIds);
+	void Sync(const Chaos::FPhysicsSolver* Solver, const TManagedArray<int32>& RigidBodyIds);
 
 	const Chaos::TVector<T, d>               & GetX                     (int32 Index) const { check(HasSyncedData(EGeometryCollectionParticlesData::X                     )); return BufferedData.GetGameDataForRead().X                     [Index]; }
 	const Chaos::TRotation<T, d>             & GetR                     (int32 Index) const { check(HasSyncedData(EGeometryCollectionParticlesData::R                     )); return BufferedData.GetGameDataForRead().R                     [Index]; }
@@ -165,7 +165,7 @@ private:
 		void Reset(EGeometryCollectionParticlesData Data);
 
 		/** Copy the specified particle information for the specified range of rigid body id. */
-		void Copy(EGeometryCollectionParticlesData Data, const Chaos::FPBDRigidsSolver* Solver, const TManagedArray<int32>& RigidBodyIds);
+		void Copy(EGeometryCollectionParticlesData Data, const Chaos::FPhysicsSolver* Solver, const TManagedArray<int32>& RigidBodyIds);
 
 		/** Return a string with the entire set of value for the synced data of the specified particle. */
 		FString ToString(int32 Index, const TCHAR* Separator) const;
@@ -179,6 +179,8 @@ private:
 	uint64 SyncFrame;
 };	
 
+extern template class TGeometryCollectionParticlesData<float, 3>;
+
 /** Current Chaos particles syncer type. */
 typedef TGeometryCollectionParticlesData<float, 3> FGeometryCollectionParticlesData;
 
diff --git a/Engine/Source/Runtime/Experimental/GeometryCollectionEngine/Public/GeometryCollection/PhysicsAssetSimulation.h b/Engine/Source/Runtime/Experimental/GeometryCollectionEngine/Public/GeometryCollection/PhysicsAssetSimulation.h
--- a/Engine/Source/Runtime/Experimental/GeometryCollectionEngine/Public/GeometryCollection/PhysicsAssetSimulation.h
+++ b/Engine/Source/Runtime/Experimental/GeometryCollectionEngine/Public/GeometryCollection/PhysicsAssetSimulation.h
@@ -9,8 +9,8 @@
 #if INCLUDE_CHAOS
 
 class AActor;
-class FSkeletalMeshPhysicsObject;
-struct FSkeletalMeshPhysicsObjectParams;
+class FSkeletalMeshPhysicsProxy;
+struct FSkeletalMeshPhysicsProxyParams;
 class UChaosPhysicalMaterial;
 class UPhysicsAsset;
 class USkeletalMeshComponent;
@@ -18,8 +18,8 @@ class USkeletalMeshComponent;
 
 struct GEOMETRYCOLLECTIONENGINE_API FPhysicsAssetSimulationUtil
 {
-	static void BuildParams(const UObject* Caller, const AActor* OwningActor, const USkeletalMeshComponent* SkelMeshComponent, const UPhysicsAsset* PhysicsAsset, FSkeletalMeshPhysicsObjectParams& OutParams);
-	static bool UpdateAnimState(const UObject* Caller, const AActor* OwningActor, const USkeletalMeshComponent* SkelMeshComponent, const float Dt, FSkeletalMeshPhysicsObjectParams& OutParams);
+	static void BuildParams(const UObject* Caller, const AActor* OwningActor, const USkeletalMeshComponent* SkelMeshComponent, const UPhysicsAsset* PhysicsAsset, FSkeletalMeshPhysicsProxyParams& OutParams);
+	static bool UpdateAnimState(const UObject* Caller, const AActor* OwningActor, const USkeletalMeshComponent* SkelMeshComponent, const float Dt, FSkeletalMeshPhysicsProxyParams& OutParams);
 };
 
 #endif
diff --git a/Engine/Source/Runtime/Experimental/GeometryCollectionEngine/Public/GeometryCollection/SkeletalMeshSimulationComponent.h b/Engine/Source/Runtime/Experimental/GeometryCollectionEngine/Public/GeometryCollection/SkeletalMeshSimulationComponent.h
--- a/Engine/Source/Runtime/Experimental/GeometryCollectionEngine/Public/GeometryCollection/SkeletalMeshSimulationComponent.h
+++ b/Engine/Source/Runtime/Experimental/GeometryCollectionEngine/Public/GeometryCollection/SkeletalMeshSimulationComponent.h
@@ -172,7 +172,7 @@ private:
 
 private:
 #if INCLUDE_CHAOS
-	FSkeletalMeshPhysicsObject* PhysicsObject;
+	FSkeletalMeshPhysicsProxy* PhysicsProxy;
 
 	//@todo(mlentine): Don't have one per static mesh
 	TUniquePtr<Chaos::TChaosPhysicsMaterial<float>> ChaosMaterial;
diff --git a/Engine/Source/Runtime/Experimental/GeometryCollectionEngine/Public/GeometryCollection/StaticMeshSimulationComponent.h b/Engine/Source/Runtime/Experimental/GeometryCollectionEngine/Public/GeometryCollection/StaticMeshSimulationComponent.h
--- a/Engine/Source/Runtime/Experimental/GeometryCollectionEngine/Public/GeometryCollection/StaticMeshSimulationComponent.h
+++ b/Engine/Source/Runtime/Experimental/GeometryCollectionEngine/Public/GeometryCollection/StaticMeshSimulationComponent.h
@@ -12,7 +12,7 @@
 
 #include "StaticMeshSimulationComponent.generated.h"
 
-class FStaticMeshPhysicsObject;
+class FStaticMeshPhysicsProxy;
 
 /**
 *	UStaticMeshSimulationComponent
@@ -109,9 +109,9 @@ public:
 private : 
 
 	/** List of physics objects this simulation component created. */
-	TArray<FStaticMeshPhysicsObject*> PhysicsObjects;
+	TArray<FStaticMeshPhysicsProxy*> PhysicsProxies;
 
-	/** List of component for which this simulation component created a physics object. Parallel array to PhysicsObjects, so PhysicsObjects[i] corresponds to SimulatedComponents[i] */
+	/** List of component for which this simulation component created a physics object. Parallel array to PhysicsProxy, so PhysicsProxies[i] corresponds to SimulatedComponents[i] */
 	UPROPERTY()
 	TArray<UPrimitiveComponent*> SimulatedComponents;
 
diff --git a/Engine/Source/Runtime/Experimental/GeometryCollectionSimulationCore/Private/GeometryCollection/GeometryCollectionCollisionStructureManager.cpp b/Engine/Source/Runtime/Experimental/GeometryCollectionSimulationCore/Private/GeometryCollection/GeometryCollectionCollisionStructureManager.cpp
--- a/Engine/Source/Runtime/Experimental/GeometryCollectionSimulationCore/Private/GeometryCollection/GeometryCollectionCollisionStructureManager.cpp
+++ b/Engine/Source/Runtime/Experimental/GeometryCollectionSimulationCore/Private/GeometryCollection/GeometryCollectionCollisionStructureManager.cpp
@@ -96,7 +96,7 @@ FCollisionStructureManager::NewSimplicial(
 		if (NumParticles)
 		{
 			int32 VertexCounter = 0;
-			Simplicial->AddElements(NumParticles);
+			Simplicial->AddParticles(NumParticles);
 			for (int32 i = 0; i<NumParticles;i++)
 			{
 				if (!OutsideVertices[i].ContainsNaN())
@@ -110,7 +110,7 @@ FCollisionStructureManager::NewSimplicial(
 
 		if(!Simplicial->Size())
 		{
-			Simplicial->AddElements(1);
+			Simplicial->AddParticles(1);
 			Simplicial->X(0) = Chaos::TVector<float, 3>(0);
 		}
 
@@ -146,7 +146,7 @@ FCollisionStructureManager::NewSimplicial(
 		if (Result.Num())
 		{
 			int32 VertexCounter = 0;
-			Simplicial->AddElements(Result.Num());
+			Simplicial->AddParticles(Result.Num());
 			for (int Index = Result.Num() - 1; 0 <= Index; Index--)
 			{
 				if (!Result[Index].ContainsNaN())
@@ -160,7 +160,7 @@ FCollisionStructureManager::NewSimplicial(
 
 		if (!Simplicial->Size())
 		{
-			Simplicial->AddElements(1);
+			Simplicial->AddParticles(1);
 			Simplicial->X(0) = Chaos::TVector<float, 3>(0);
 		}
 
diff --git a/Engine/Source/Runtime/Landscape/Private/LandscapeCollision.cpp b/Engine/Source/Runtime/Landscape/Private/LandscapeCollision.cpp
--- a/Engine/Source/Runtime/Landscape/Private/LandscapeCollision.cpp
+++ b/Engine/Source/Runtime/Landscape/Private/LandscapeCollision.cpp
@@ -49,6 +49,11 @@
 	#include "IPhysXCooking.h"
 #endif
 
+#if WITH_CHAOS
+#include "PhysXToChaosUtil.h"
+#include "Chaos/ParticleHandle.h"
+#endif
+
 using namespace PhysicsInterfaceTypes;
 
 #if ENABLE_COOK_STATS
@@ -251,6 +256,31 @@ void ULandscapeHeightfieldCollisionComponent::OnCreatePhysicsState()
 
 				HeightFieldActorSync->attachShape(*HeightFieldShapeSync);
 
+#if WITH_CHAOS
+				FActorCreationParams Params;
+				Params.InitialTM = LandscapeShapeTM;
+				Params.bQueryOnly = true;
+				Params.Scene = GetWorld()->GetPhysicsScene();
+				FPhysicsActorHandle PhysHandle;
+				FPhysicsInterface::CreateActor(Params, PhysHandle);
+
+				TUniquePtr<Chaos::TImplicitObjectTransformed<float, 3>> ChaosHeightField = PxShapeToChaosGeom(HeightFieldShapeSync);
+				TUniquePtr<Chaos::TPerShapeData<float, 3>> NewShape = MakeUnique<Chaos::TPerShapeData<float, 3>>();
+
+				NewShape->Geometry = ChaosHeightField.Get();
+				NewShape->QueryData = QueryFilterData;
+				NewShape->SimData = SimFilterData;
+
+				Chaos::TShapesArray<float, 3> ShapeArray;
+				ShapeArray.Emplace(MoveTemp(NewShape));
+				PhysHandle->SetGeometry(MoveTemp(ChaosHeightField));
+				PhysHandle->SetShapesArray(MoveTemp(ShapeArray));
+
+				FTransform ActorTM(LandscapeComponentMatrix);
+				PhysHandle->SetX(ActorTM.GetTranslation());
+				PhysHandle->SetR(ActorTM.GetRotation());
+#endif
+
 				// attachShape holds its own ref(), so release this here.
 				HeightFieldShapeSync->release();
 
@@ -286,7 +316,7 @@ void ULandscapeHeightfieldCollisionComponent::OnCreatePhysicsState()
 					PxHeightFieldGeometry LandscapeComponentGeomEd(HeightfieldRef->RBHeightfieldEd, PxMeshGeometryFlags(), LandscapeScale.Z * LANDSCAPE_ZSCALE, LandscapeScale.Y * CollisionScale, LandscapeScale.X * CollisionScale);
 					if (LandscapeComponentGeomEd.isValid())
 					{
-#if WITH_CHAOS || WITH_IMMEDIATE_PHYSX || PHYSICS_INTERFACE_LLIMMEDIATE
+#if WITH_CHAOS || WITH_IMMEDIATE_PHYSX
                         ensure(false);
 #else
 						FPhysicsMaterialHandle_PhysX MaterialHandle = GEngine->DefaultPhysMaterial->GetPhysicsMaterial();
@@ -317,14 +347,16 @@ void ULandscapeHeightfieldCollisionComponent::OnCreatePhysicsState()
 				FPhysScene* PhysScene = GetWorld()->GetPhysicsScene();
 
 				// Set body instance data
-				BodyInstance.PhysxUserData = FPhysxUserData(&BodyInstance);
+				BodyInstance.PhysicsUserData = FPhysicsUserData(&BodyInstance);
 				BodyInstance.OwnerComponent = this;
 
-#if WITH_CHAOS || WITH_IMMEDIATE_PHYSX || PHYSICS_INTERFACE_LLIMMEDIATE
-                ensure(false);
+#if WITH_CHAOS || WITH_IMMEDIATE_PHYSX
+				TArray<FPhysicsActorHandle> Actors;
+				Actors.Add(PhysHandle);
+				PhysScene->AddActorsToScene_AssumesLocked(Actors);
 #else
 				BodyInstance.ActorHandle.SyncActor = HeightFieldActorSync;
-				HeightFieldActorSync->userData = &BodyInstance.PhysxUserData;
+				HeightFieldActorSync->userData = &BodyInstance.PhysicsUserData;
 
 				// Add to scenes
 				PxScene* SyncScene = PhysScene->GetPxScene();
@@ -407,7 +439,7 @@ void ULandscapeHeightfieldCollisionComponent::CreateCollisionObject()
 
 				for (UPhysicalMaterial* PhysicalMaterial : CookedPhysicalMaterials)
 				{
-#if WITH_CHAOS || WITH_IMMEDIATE_PHYSX || PHYSICS_INTERFACE_LLIMMEDIATE
+#if WITH_CHAOS || WITH_IMMEDIATE_PHYSX
                     ensure(false);
 #else
 					const FPhysicsMaterialHandle_PhysX& MaterialHandle = PhysicalMaterial->GetPhysicsMaterial();
@@ -909,7 +941,7 @@ void ULandscapeMeshCollisionComponent::CreateCollisionObject()
 
 				for (UPhysicalMaterial* PhysicalMaterial : CookedPhysicalMaterials)
 				{
-#if WITH_CHAOS || WITH_IMMEDIATE_PHYSX || PHYSICS_INTERFACE_LLIMMEDIATE
+#if WITH_CHAOS || WITH_IMMEDIATE_PHYSX
                     ensure(false);
 #else
 					MeshRef->UsedPhysicalMaterialArray.Add(PhysicalMaterial->GetPhysicsMaterial().Material);
@@ -1012,7 +1044,7 @@ void ULandscapeMeshCollisionComponent::OnCreatePhysicsState()
 					PTriMeshGeomEd.scale.scale.z = LandscapeScale.Z;
 					if (PTriMeshGeomEd.isValid())
 					{
-#if WITH_CHAOS || WITH_IMMEDIATE_PHYSX || PHYSICS_INTERFACE_LLIMMEDIATE
+#if WITH_CHAOS || WITH_IMMEDIATE_PHYSX
                         ensure(false);
 #else
 						PxMaterial* PDefaultMat = GEngine->DefaultPhysMaterial->GetPhysicsMaterial().Material;
@@ -1037,14 +1069,14 @@ void ULandscapeMeshCollisionComponent::OnCreatePhysicsState()
 #endif// WITH_EDITOR
 
 				// Set body instance data
-				BodyInstance.PhysxUserData = FPhysxUserData(&BodyInstance);
+				BodyInstance.PhysicsUserData = FPhysicsUserData(&BodyInstance);
 				BodyInstance.OwnerComponent = this;
 
-#if WITH_CHAOS || WITH_IMMEDIATE_PHYSX || PHYSICS_INTERFACE_LLIMMEDIATE
+#if WITH_CHAOS || WITH_IMMEDIATE_PHYSX
                 ensure(false);
 #else
 				BodyInstance.ActorHandle.SyncActor = MeshActorSync;
-				MeshActorSync->userData = &BodyInstance.PhysxUserData;
+				MeshActorSync->userData = &BodyInstance.PhysicsUserData;
 
 				// Add to scenes
 				PxScene* SyncScene = PhysScene->GetPxScene();
@@ -1117,7 +1149,7 @@ void ULandscapeHeightfieldCollisionComponent::UpdateHeightfieldRegion(int32 Comp
 			return;
 		}
 
-#if WITH_CHAOS || WITH_IMMEDIATE_PHYSX || PHYSICS_INTERFACE_LLIMMEDIATE
+#if WITH_CHAOS || WITH_IMMEDIATE_PHYSX
         ensure(false);
 #else
 		if (BodyInstance.ActorHandle.SyncActor == NULL)
diff --git a/Engine/Source/Runtime/MRMesh/Private/MRMeshComponent.cpp b/Engine/Source/Runtime/MRMesh/Private/MRMeshComponent.cpp
--- a/Engine/Source/Runtime/MRMesh/Private/MRMeshComponent.cpp
+++ b/Engine/Source/Runtime/MRMesh/Private/MRMeshComponent.cpp
@@ -727,13 +727,12 @@ UBodySetup* UMRMeshComponent::GetBodySetup()
 
 void UMRMeshComponent::SendBrickData_Internal(IMRMesh::FSendBrickDataArgs Args)
 {
-#if WITH_PHYSX
 	check(IsInGameThread());
-
-	UE_LOG(LogMrMesh, Log, TEXT("SendBrickData_Internal() processing brick %llu with %i triangles"), Args.BrickId, Args.Indices.Num() / 3);
-
 	const bool bHasBrickData = Args.Indices.Num() > 0 && Args.PositionData.Num() > 0;
 
+#if WITH_PHYSX && PHYSICS_INTERFACE_PHYSX
+	UE_LOG(LogMrMesh, Log, TEXT("SendBrickData_Internal() processing brick %llu with %i triangles"), Args.BrickId, Args.Indices.Num() / 3);
+	
 	if (!IsPendingKill() && !bNeverCreateCollisionMesh)
 	{
 		// Physics update
@@ -782,7 +781,7 @@ void UMRMeshComponent::SendBrickData_Internal(IMRMesh::FSendBrickDataArgs Args)
 				CookHelper.CookInfo = CookInfo;
 				CookHelper.CreatePhysicsMeshes_Concurrent();
 
-				MyBS->FinishCreatingPhysicsMeshes(CookHelper.OutNonMirroredConvexMeshes, CookHelper.OutMirroredConvexMeshes, CookHelper.OutTriangleMeshes);
+				MyBS->FinishCreatingPhysicsMeshes_PhysX(CookHelper.OutNonMirroredConvexMeshes, CookHelper.OutMirroredConvexMeshes, CookHelper.OutTriangleMeshes);
 
 				FBodyInstance* MyBI = BodyInstances[BodyIndex];
 				MyBI->TermBody();
diff --git a/Engine/Source/Runtime/NavigationSystem/NavigationSystem.Build.cs b/Engine/Source/Runtime/NavigationSystem/NavigationSystem.Build.cs
--- a/Engine/Source/Runtime/NavigationSystem/NavigationSystem.Build.cs
+++ b/Engine/Source/Runtime/NavigationSystem/NavigationSystem.Build.cs
@@ -25,8 +25,6 @@ namespace UnrealBuildTool.Rules
                     "Core",
                     "CoreUObject",
                     "Engine",
-					"GeometryCollectionCore",
-					"GeometryCollectionEngine",
                 }
                 );
 
@@ -51,6 +49,16 @@ namespace UnrealBuildTool.Rules
 
             SetupModulePhysicsSupport(Target);
 
+			if (Target.bCompileChaos || Target.bUseChaos)
+            {
+                PublicDependencyModuleNames.AddRange(
+                    new string[] {
+                        "GeometryCollectionCore",
+                        "GeometryCollectionEngine",
+                    }
+					);
+            }
+
             if (Target.bCompileRecast)
             {
                 PrivateDependencyModuleNames.Add("Navmesh");
diff --git a/Engine/Source/Runtime/PhysicsCore/PhysicsCore.Build.cs b/Engine/Source/Runtime/PhysicsCore/PhysicsCore.Build.cs
--- a/Engine/Source/Runtime/PhysicsCore/PhysicsCore.Build.cs
+++ b/Engine/Source/Runtime/PhysicsCore/PhysicsCore.Build.cs
@@ -11,8 +11,7 @@ public class PhysicsCore: ModuleRules
         PrivateDependencyModuleNames.AddRange(
 			new string[] {
 				"Core",
-				"CoreUObject",
-				"Chaos"
+				"CoreUObject"
 			}
 		);
 
diff --git a/Engine/Source/Runtime/PhysicsCore/Private/PhysTestSerializer.cpp b/Engine/Source/Runtime/PhysicsCore/Private/PhysTestSerializer.cpp
--- a/Engine/Source/Runtime/PhysicsCore/Private/PhysTestSerializer.cpp
+++ b/Engine/Source/Runtime/PhysicsCore/Private/PhysTestSerializer.cpp
@@ -3,7 +3,6 @@
 // Physics engine integration utilities
 
 #include "PhysTestSerializer.h"
-#if !WITH_CHAOS_NEEDS_TO_BE_FIXED
 
 #if WITH_PHYSX
 #include "PhysXIncludes.h"
@@ -199,30 +198,49 @@ void FPhysTestSerializer::CreateChaosData()
 	if (NumStatic)
 	{
 		Scene->getActors(PxActorTypeFlag::eRIGID_STATIC, Actors.GetData(), NumStatic);
+		auto NewParticles = Particles.CreateStaticParticles(NumStatic);	//question: do we want to distinguish query only and sim only actors?
+		for (uint32 Idx = 0; Idx < NumStatic; ++Idx)
+		{
+			GTParticles.Emplace(MakeUnique<TGeometryParticle<float, 3>>());
+			NewParticles[Idx]->GTGeometryParticle() = GTParticles.Last().Get();
+		}
 	}
 
 	if (NumDynamic)
 	{
 		Scene->getActors(PxActorTypeFlag::eRIGID_DYNAMIC, &Actors[NumStatic], NumDynamic);
-	}
+		auto NewParticles = Particles.CreateDynamicParticles(NumDynamic);	//question: do we want to distinguish query only and sim only actors?
 
-	TPBDRigidParticles<float, 3> Particles;
-	Particles.AddParticles(NumActors);	//question: do we want to distinguish query only and sim only actors?
+		for (uint32 Idx = 0; Idx < NumDynamic; ++Idx)
+		{
+			GTParticles.Emplace(MakeUnique<TPBDRigidParticle<float, 3>>());
+			NewParticles[Idx]->GTGeometryParticle() = GTParticles.Last().Get();
+		}
+	}
 
+	auto& Handles = Particles.GetParticleHandles();
 	int32 Idx = 0;
 	for (PxActor* Act : Actors)
 	{
 		//transform
 		PxRigidActor* Actor = static_cast<PxRigidActor*>(Act);
-		Particles.X(Idx) = P2UVector(Actor->getGlobalPose().p);
-		Particles.R(Idx) = P2UQuat(Actor->getGlobalPose().q);
-		Particles.V(Idx) = TVector<float, 3>(0);
-		Particles.W(Idx) + TVector<float, 3>(0);
-		Particles.P(Idx) = Particles.X(Idx);
-		Particles.Q(Idx) = Particles.R(Idx);
-		Particles.SetDisabledLowLevel(Idx, false);
+		auto& Particle = Handles.Handle(Idx);
+		auto& GTParticle = Particle->GTGeometryParticle();
+		Particle->X() = P2UVector(Actor->getGlobalPose().p);
+		Particle->R() = P2UQuat(Actor->getGlobalPose().q);
+		Particle->GTGeometryParticle()->SetX(Particle->X());
+		Particle->GTGeometryParticle()->SetR(Particle->R());
+
+		if (auto PBDRigid = Particle->AsDynamic())
+		{
+			PBDRigid->P() = Particle->X();
+			PBDRigid->Q() = Particle->R();
 
-		PxActorToChaosIdx.Add(Act, Idx);
+			PBDRigid->GTGeometryParticle()->AsDynamic()->SetP(PBDRigid->P());
+			PBDRigid->GTGeometryParticle()->AsDynamic()->SetQ(PBDRigid->R());
+		}
+
+		PxActorToChaosHandle.Add(Act, Particle.Get());
 
 		//geometry
 		TArray<TUniquePtr<TImplicitObject<float, 3>>> Geoms;
@@ -235,7 +253,6 @@ void FPhysTestSerializer::CreateChaosData()
 			if (TUniquePtr<TImplicitObjectTransformed<float, 3>> Geom = PxShapeToChaosGeom(Shape))
 			{
 				Geoms.Add(MoveTemp(Geom));
-				PxShapeToChaosImplicit.Add(Shape, Geoms.Last().Get());
 			}
 		}
 
@@ -243,19 +260,27 @@ void FPhysTestSerializer::CreateChaosData()
 		{
 			if (Geoms.Num() == 1)
 			{
-				Particles.SetDynamicGeometry(Idx, MoveTemp(Geoms[0]));
+				auto SharedGeom = TSharedPtr<TImplicitObject<float, 3>, ESPMode::ThreadSafe>(Geoms[0].Release());
+				GTParticle->SetGeometry(SharedGeom);
+				Particle->SetSharedGeometry(SharedGeom);
 			}
 			else
 			{
-				Particles.SetDynamicGeometry(Idx, MakeUnique<TImplicitObjectUnion<float, 3>>(MoveTemp(Geoms)));
+				GTParticle->SetGeometry(MakeUnique<TImplicitObjectUnion<float, 3>>(MoveTemp(Geoms)));
+				Particle->SetGeometry(GTParticle->Geometry());
 			}
 		}
+
+		int32 ShapeIdx = 0;
+		for (PxShape* Shape : Shapes)
+		{
+			PxShapeToChaosShapes.Add(Shape, GTParticle->ShapesArray()[ShapeIdx++].Get());
+		}
+
 		++Idx;
 	}
 
-	ChaosEvolution = MakeUnique<TPBDRigidsEvolutionGBF<float, 3>>(MoveTemp(Particles));
+	ChaosEvolution = MakeUnique<TPBDRigidsEvolutionGBF<float, 3>>(Particles);
 	bChaosDataReady = true;
 #endif
 }
-
-#endif // !WITH_CHAOS_NEEDS_TO_BE_FIXED
\ No newline at end of file
diff --git a/Engine/Source/Runtime/PhysicsCore/Private/PhysicsInterfaceUtilsCore.cpp b/Engine/Source/Runtime/PhysicsCore/Private/PhysicsInterfaceUtilsCore.cpp
--- a/Engine/Source/Runtime/PhysicsCore/Private/PhysicsInterfaceUtilsCore.cpp
+++ b/Engine/Source/Runtime/PhysicsCore/Private/PhysicsInterfaceUtilsCore.cpp
@@ -72,12 +72,96 @@
 	return ECollisionShapeType::None;
 }
 
+FPhysXShapeAdapter::FPhysXShapeAdapter(const FQuat& Rot, const FCollisionShape& CollisionShape)
+	: Rotation(FQuat::Identity)
+{
+	// Perform other kinds of zero-extent queries as zero-extent sphere queries
+	if ((CollisionShape.ShapeType != ECollisionShape::Sphere) && CollisionShape.IsNearlyZero())
+	{
+		GeometryHolder.storeAny(PxSphereGeometry(FCollisionShape::MinSphereRadius()));
+	}
+	else
+	{
+		switch (CollisionShape.ShapeType)
+		{
+		case ECollisionShape::Box:
+		{
+			PxVec3 BoxExtents = U2PVector(CollisionShape.GetBox());
+			BoxExtents.x = FMath::Max(BoxExtents.x, FCollisionShape::MinBoxExtent());
+			BoxExtents.y = FMath::Max(BoxExtents.y, FCollisionShape::MinBoxExtent());
+			BoxExtents.z = FMath::Max(BoxExtents.z, FCollisionShape::MinBoxExtent());
+
+			GeometryHolder.storeAny(PxBoxGeometry(BoxExtents));
+			Rotation = Rot;
+			break;
+		}
+		case ECollisionShape::Sphere:
+			GeometryHolder.storeAny(PxSphereGeometry(FMath::Max(CollisionShape.GetSphereRadius(), FCollisionShape::MinSphereRadius())));
+			break;
+		case ECollisionShape::Capsule:
+		{
+			const float CapsuleRadius = CollisionShape.GetCapsuleRadius();
+			const float CapsuleHalfHeight = CollisionShape.GetCapsuleHalfHeight();
+			if (CapsuleRadius < CapsuleHalfHeight)
+			{
+				GeometryHolder.storeAny(PxCapsuleGeometry(FMath::Max(CapsuleRadius, FCollisionShape::MinCapsuleRadius()), FMath::Max(CollisionShape.GetCapsuleAxisHalfLength(), FCollisionShape::MinCapsuleAxisHalfHeight())));
+				Rotation = ConvertToPhysXCapsuleRot(Rot);
+			}
+			else
+			{
+				// Use a sphere instead.
+				GeometryHolder.storeAny(PxSphereGeometry(FMath::Max(CapsuleRadius, FCollisionShape::MinSphereRadius())));
+			}
+			break;
+		}
+		default:
+			// invalid type
+			ensure(false);
+		}
+	}
+}
+
+static const PxQuat CapsuleRotator(0.f, 0.707106781f, 0.f, 0.707106781f);
+
+FQuat ConvertToPhysXCapsuleRot(const FQuat& GeomRot)
+{
+	// Rotation required because PhysX capsule points down X, we want it down Z
+	return P2UQuat(U2PQuat(GeomRot) * CapsuleRotator);	//todo(ocohen): avoid pointless conversion
+}
+
+FQuat ConvertToUECapsuleRot(const PxQuat & PGeomRot)
+{
+	return P2UQuat(PGeomRot * CapsuleRotator.getConjugate());
+}
+
+FQuat ConvertToUECapsuleRot(const FQuat & GeomRot)
+{
+	return GeomRot * P2UQuat(CapsuleRotator.getConjugate());
+}
+
+PxTransform ConvertToPhysXCapsulePose(const FTransform& GeomPose)
+{
+	PxTransform PFinalPose;
+
+	PFinalPose.p = U2PVector(GeomPose.GetTranslation());
+	// Rotation required because PhysX capsule points down X, we want it down Z
+	PFinalPose.q = U2PQuat(ConvertToPhysXCapsuleRot(GeomPose.GetRotation()));
+	return PFinalPose;
+}
+
+#if USE_SCENE_LOCK
+DEFINE_STAT(STAT_PhysSceneReadLock);
+DEFINE_STAT(STAT_PhysSceneWriteLock);
+#endif
+
+#endif
+
 uint32 FindFaceIndex(const FHitLocation& PHit, const FVector& UnitDir)
 {
 #if PHYSICS_INTERFACE_PHYSX
 	PxConvexMeshGeometry convexGeom;
 	PxVec3 PUnitDir = U2PVector(UnitDir);
-	if (PHit.shape->getConvexMeshGeometry(convexGeom))
+	if(PHit.shape->getConvexMeshGeometry(convexGeom))
 	{
 		//PhysX has given us the most correct face. However, we actually want the most useful face which is the one with the most opposed normal within some radius.
 		//So for example, if we are sweeping against a corner we should take the plane that is most opposing, even if it's not the exact one we hit.
@@ -86,15 +170,15 @@
 		PxTransform actorPose;
 		bool bFoundKinematicTarget = false;
 
-		if (PxRigidDynamic* rigidDynamic = PHit.actor->is<PxRigidDynamic>())
+		if(PxRigidDynamic* rigidDynamic = PHit.actor->is<PxRigidDynamic>())
 		{
-			if (rigidDynamic->getRigidBodyFlags() & PxRigidBodyFlag::eUSE_KINEMATIC_TARGET_FOR_SCENE_QUERIES)
+			if(rigidDynamic->getRigidBodyFlags() & PxRigidBodyFlag::eUSE_KINEMATIC_TARGET_FOR_SCENE_QUERIES)
 			{
 				bFoundKinematicTarget = rigidDynamic->getKinematicTarget(actorPose);
 			}
 		}
 
-		if (!bFoundKinematicTarget)
+		if(!bFoundKinematicTarget)
 		{
 			actorPose = PHit.actor->getGlobalPose();
 		}
@@ -130,7 +214,7 @@
 			PxU32 maxDIndex = 0;
 			PxReal minNormalDot = PX_MAX_REAL;
 
-			for (PxU32 j = 0; j < nbPolys; j++)
+			for(PxU32 j = 0; j < nbPolys; j++)
 			{
 				PxHullPolygon hullPolygon;
 				convexGeom.convexMesh->getPolygonData(j, hullPolygon);
@@ -144,7 +228,7 @@
 
 				PxReal d = plane.distance(localPoint);
 				// Track plane that impact point is furthest point (will be out fallback normal)
-				if (d > maxD)
+				if(d > maxD)
 				{
 					maxDIndex = j;
 					maxD = d;
@@ -152,13 +236,13 @@
 
 				//Because we are searching against a convex hull, we will never get multiple faces that are both in front of the contact point _and_ have an opposing normal (except the face we hit).
 				//However, we may have just missed a plane which is now "behind" the contact point while still being inside the radius
-				if (d < -FindFaceInRadius)
+				if(d < -FindFaceInRadius)
 					continue;
 
 				// Calculate direction dot plane normal
 				const PxReal normalDot = plane.n.dot(localDir);
 				// If this is more opposing than our current 'most opposing' normal, update 'most opposing'
-				if (normalDot < minNormalDot)
+				if(normalDot < minNormalDot)
 				{
 					minIndex = j;
 					bMinIndexValid = true;
@@ -167,7 +251,7 @@
 			}
 
 			// If we found at least one face that we are considered 'on', use best normal
-			if (bMinIndexValid)
+			if(bMinIndexValid)
 			{
 				return minIndex;
 			}
@@ -184,87 +268,3 @@
 	ensure(false);
 	return 0;
 }
-
-FPhysXShapeAdaptor::FPhysXShapeAdaptor(const FQuat& Rot, const FCollisionShape& CollisionShape)
-	: Rotation(FQuat::Identity)
-{
-	// Perform other kinds of zero-extent queries as zero-extent sphere queries
-	if ((CollisionShape.ShapeType != ECollisionShape::Sphere) && CollisionShape.IsNearlyZero())
-	{
-		GeometryHolder.storeAny(PxSphereGeometry(FCollisionShape::MinSphereRadius()));
-	}
-	else
-	{
-		switch (CollisionShape.ShapeType)
-		{
-		case ECollisionShape::Box:
-		{
-			PxVec3 BoxExtents = U2PVector(CollisionShape.GetBox());
-			BoxExtents.x = FMath::Max(BoxExtents.x, FCollisionShape::MinBoxExtent());
-			BoxExtents.y = FMath::Max(BoxExtents.y, FCollisionShape::MinBoxExtent());
-			BoxExtents.z = FMath::Max(BoxExtents.z, FCollisionShape::MinBoxExtent());
-
-			GeometryHolder.storeAny(PxBoxGeometry(BoxExtents));
-			Rotation = Rot;
-			break;
-		}
-		case ECollisionShape::Sphere:
-			GeometryHolder.storeAny(PxSphereGeometry(FMath::Max(CollisionShape.GetSphereRadius(), FCollisionShape::MinSphereRadius())));
-			break;
-		case ECollisionShape::Capsule:
-		{
-			const float CapsuleRadius = CollisionShape.GetCapsuleRadius();
-			const float CapsuleHalfHeight = CollisionShape.GetCapsuleHalfHeight();
-			if (CapsuleRadius < CapsuleHalfHeight)
-			{
-				GeometryHolder.storeAny(PxCapsuleGeometry(FMath::Max(CapsuleRadius, FCollisionShape::MinCapsuleRadius()), FMath::Max(CollisionShape.GetCapsuleAxisHalfLength(), FCollisionShape::MinCapsuleAxisHalfHeight())));
-				Rotation = ConvertToPhysXCapsuleRot(Rot);
-			}
-			else
-			{
-				// Use a sphere instead.
-				GeometryHolder.storeAny(PxSphereGeometry(FMath::Max(CapsuleRadius, FCollisionShape::MinSphereRadius())));
-			}
-			break;
-		}
-		default:
-			// invalid type
-			ensure(false);
-		}
-	}
-}
-
-static const PxQuat CapsuleRotator(0.f, 0.707106781f, 0.f, 0.707106781f);
-
-FQuat ConvertToPhysXCapsuleRot(const FQuat& GeomRot)
-{
-	// Rotation required because PhysX capsule points down X, we want it down Z
-	return P2UQuat(U2PQuat(GeomRot) * CapsuleRotator);	//todo(ocohen): avoid pointless conversion
-}
-
-FQuat ConvertToUECapsuleRot(const PxQuat & PGeomRot)
-{
-	return P2UQuat(PGeomRot * CapsuleRotator.getConjugate());
-}
-
-FQuat ConvertToUECapsuleRot(const FQuat & GeomRot)
-{
-	return GeomRot * P2UQuat(CapsuleRotator.getConjugate());
-}
-
-PxTransform ConvertToPhysXCapsulePose(const FTransform& GeomPose)
-{
-	PxTransform PFinalPose;
-
-	PFinalPose.p = U2PVector(GeomPose.GetTranslation());
-	// Rotation required because PhysX capsule points down X, we want it down Z
-	PFinalPose.q = U2PQuat(ConvertToPhysXCapsuleRot(GeomPose.GetRotation()));
-	return PFinalPose;
-}
-
-#if USE_SCENE_LOCK
-DEFINE_STAT(STAT_PhysSceneReadLock);
-DEFINE_STAT(STAT_PhysSceneWriteLock);
-#endif
-
-#endif
diff --git a/Engine/Source/Runtime/PhysicsCore/Private/SQCapture.cpp b/Engine/Source/Runtime/PhysicsCore/Private/SQCapture.cpp
--- a/Engine/Source/Runtime/PhysicsCore/Private/SQCapture.cpp
+++ b/Engine/Source/Runtime/PhysicsCore/Private/SQCapture.cpp
@@ -1,13 +1,15 @@
 // Copyright 1998-2019 Epic Games, Inc. All Rights Reserved.
 
 #include "SQCapture.h"
+
+#if INCLUDE_CHAOS
 #include "Chaos/ImplicitObject.h"
+#endif
+
 #include "PhysXSupportCore.h"
 #include "PhysicsPublicCore.h"
 #include "PhysTestSerializer.h" 
 
-#if  !WITH_CHAOS_NEEDS_TO_BE_FIXED
-
 #if WITH_PHYSX && INCLUDE_CHAOS
 #include "PhysXToChaosUtil.h"
 #endif
@@ -70,7 +72,7 @@ void FSQCapture::SerializePhysXHitType(FArchive& Ar, T& Hit)
 }
 
 template <typename THit>
-void FSQCapture::SerializePhysXBuffers(FArchive& Ar, int32 Version, FDynamicHitBuffer<THit>& PhysXBuffer)
+void FSQCapture::SerializePhysXBuffers(FArchive& Ar, int32 Version, PhysXInterface::FDynamicHitBuffer<THit>& PhysXBuffer)
 {
 	Ar << PhysXBuffer.hasBlock;
 	if (PhysXBuffer.hasBlock)
@@ -109,7 +111,7 @@ void FSQCapture::SerializePhysXBuffers(FArchive& Ar, int32 Version, FDynamicHitB
 
 void FSQCapture::SerializeActorToShapeHitsArray(FArchive& Ar)
 {
-	int32 NumActors = ActorToShapeHitsArray.Num();
+	int32 NumActors = PxActorToShapeHitsArray.Num();
 	Ar << NumActors;
 	if (Ar.IsLoading())
 	{
@@ -133,12 +135,12 @@ void FSQCapture::SerializeActorToShapeHitsArray(FArchive& Ar)
 			}
 			PxActor* ActorPtr = static_cast<PxActor*>(PhysSerializer.FindObject(Actor));
 			check(ActorPtr);
-			ActorToShapeHitsArray.Add(ActorPtr, Pairs);
+			PxActorToShapeHitsArray.Add(ActorPtr, Pairs);
 		}
 	}
 	else if (Ar.IsSaving())
 	{
-		for (auto& Itr : ActorToShapeHitsArray)
+		for (auto& Itr : PxActorToShapeHitsArray)
 		{
 			uint64 Actor = (PxSerialObjectId)NonTransientToTransientActors.FindChecked(Itr.Key);
 			Ar << Actor;
@@ -211,11 +213,13 @@ void FSQCapture::Serialize(FArchive& Ar)
 	{
 		check(false);	//not supported yet
 	}
+#if WITH_PHYSX
 	if (Ar.IsLoading())
 	{
 		CreateChaosData();
 	}
 #endif
+#endif
 }
 
 #if WITH_PHYSX
@@ -297,14 +301,14 @@ void FSQCapture::CapturePhysXFilterResults(const PxScene& TransientScene, const
 		}
 
 		PxActor* NonTransientActor = static_cast<PxActor*>(PhysSerializer.FindObject((PxSerialObjectId)TransientActor));
-		ActorToShapeHitsArray.Add(NonTransientActor, ShapeHitsArray);
+		PxActorToShapeHitsArray.Add(NonTransientActor, ShapeHitsArray);
 		NonTransientToTransientActors.Add(NonTransientActor, TransientActor);
 	}
 }
 #endif
 
 template <typename THit>
-void EndCaptureHelper(FDynamicHitBuffer<THit>& Dest, const PxHitCallback<THit>& Results)
+void EndCaptureHelper(PhysXInterface::FDynamicHitBuffer<THit>& Dest, const PxHitCallback<THit>& Results)
 {
 	Dest.block = Results.block;
 	Dest.hasBlock = Results.hasBlock;
@@ -321,7 +325,7 @@ void EndCaptureHelper(FDynamicHitBuffer<THit>& Dest, const PxHitCallback<THit>&
 	else
 	{
 		//we know this came from a dynamic hit buffer because that's how UE uses the physx api. Since it's dynamic block is already in the hits buffer
-		const FDynamicHitBuffer<THit>& DynamicResults = static_cast<const FDynamicHitBuffer<THit>&>(Results);
+		const PhysXInterface::FDynamicHitBuffer<THit>& DynamicResults = static_cast<const PhysXInterface::FDynamicHitBuffer<THit>&>(Results);
 		Dest.processTouches(DynamicResults.GetHits(), DynamicResults.GetNumHits());
 	}
 }
@@ -363,7 +367,7 @@ void FSQCapture::StartCaptureChaos(const Chaos::TImplicitObject<float, 3>& InQue
 	}
 }
 
-void FSQCapture::EndCaptureChaos(const FPhysicsHitCallback<FSweepHit>& Results)
+void FSQCapture::EndCaptureChaos(const ChaosInterface::FSQHitBuffer<ChaosInterface::FSweepHit>& Results)
 {
 	if (IsInGameThread())
 	{
@@ -376,7 +380,7 @@ void FSQCapture::EndCaptureChaos(const FPhysicsHitCallback<FSweepHit>& Results)
 #if WITH_PHYSX
 
 template <typename THit>
-void FixupBufferPointers(FPhysTestSerializer& PhysSerializer, FDynamicHitBuffer<THit>& PhysXBuffer)
+void FixupBufferPointers(FPhysTestSerializer& PhysSerializer, PhysXInterface::FDynamicHitBuffer<THit>& PhysXBuffer)
 {
 	auto FixupPointersLambda = [&PhysSerializer](PxActorShape& Hit)
 	{
@@ -397,12 +401,12 @@ void FixupBufferPointers(FPhysTestSerializer& PhysSerializer, FDynamicHitBuffer<
 	}
 }
 
-#if WITH_PHYSX
-ECollisionQueryHitType FSQCapture::GetFilterResult(const PxShape* Shape, const PxActor* Actor) const
+template <typename TShape, typename TActor>
+ECollisionQueryHitType GetFilterResultHelper(const TShape* Shape, const TActor* Actor, const TMap<TActor*, TArray<TPair<TShape*, ECollisionQueryHitType>>>& ActorToShapeHitsArray)
 {
-	if (const TArray<TPair<PxShape*, ECollisionQueryHitType>>* ActorToPairs = ActorToShapeHitsArray.Find(Actor))
+	if (const TArray<TPair<TShape*, ECollisionQueryHitType>>* ActorToPairs = ActorToShapeHitsArray.Find(Actor))
 	{
-		for (const TPair<PxShape*, ECollisionQueryHitType>& Pair : *ActorToPairs)
+		for (const TPair<TShape*, ECollisionQueryHitType>& Pair : *ActorToPairs)
 		{
 			if (Pair.Key == Shape)
 			{
@@ -414,13 +418,16 @@ ECollisionQueryHitType FSQCapture::GetFilterResult(const PxShape* Shape, const P
 	ensure(false);	//should not get here, means we didn't properly capture all filter results
 	return ECollisionQueryHitType::Block;
 }
-#endif
+
+ECollisionQueryHitType FSQCapture::GetFilterResult(const PxShape* Shape, const PxActor* Actor) const
+{
+	return GetFilterResultHelper(Shape, Actor, PxActorToShapeHitsArray);
+}
 
 #if INCLUDE_CHAOS
-ECollisionQueryHitType FSQCapture::GetFilterResult(const Chaos::TImplicitObject<float, 3>& Shape, const int32 ActorIdx, const int32 ShapeIdx) const
+ECollisionQueryHitType FSQCapture::GetFilterResult(const Chaos::TPerShapeData<float,3>* Shape, const Chaos::TGeometryParticle<float,3>* Actor) const
 {
-	const TArray<ECollisionQueryHitType>& ActorToResults = ChaosActorToHitsArray.FindChecked(ActorIdx);
-	return ActorToResults[ShapeIdx];
+	return GetFilterResultHelper(Shape, Actor, ChaosActorToShapeHitsArray);
 }
 #endif
 
@@ -429,18 +436,18 @@ class FSQCaptureFilterCallback : public ICollisionQueryFilterCallbackBase
 public:
 	FSQCaptureFilterCallback(const FSQCapture& InCapture) : Capture(InCapture) {}
 	virtual ~FSQCaptureFilterCallback() {}
-	virtual ECollisionQueryHitType PostFilter(const FCollisionFilterData& FilterData, const FPhysicsQueryHit& Hit) override { check(false);  return ECollisionQueryHitType::Block; }
-	virtual ECollisionQueryHitType PreFilter(const FCollisionFilterData& FilterData, const FPhysicsShape& Shape, const FPhysicsActor& Actor) override { check(false); return ECollisionQueryHitType::Block; }
+#if INCLUDE_CHAOS
+	virtual ECollisionQueryHitType PostFilter(const FCollisionFilterData& FilterData, const ChaosInterface::FQueryHit& Hit) override { check(false);  return ECollisionQueryHitType::Touch; }
+	virtual ECollisionQueryHitType PreFilter(const FCollisionFilterData& FilterData, const Chaos::TPerShapeData<float, 3>& Shape, const Chaos::TGeometryParticle<float,3>& Actor) override { return Capture.GetFilterResult(&Shape, &Actor); }
+#endif
 
 #if WITH_PHYSX
+	virtual ECollisionQueryHitType PostFilter(const FCollisionFilterData& FilterData, const physx::PxQueryHit& Hit) override { return ECollisionQueryHitType::Touch; }
+	virtual ECollisionQueryHitType PreFilter(const FCollisionFilterData& FilterData, const physx::PxShape& Shape, physx::PxRigidActor& Actor) override { return Capture.GetFilterResult(&Shape, &Actor); }
 	virtual PxQueryHitType::Enum preFilter(const PxFilterData& filterData, const PxShape* shape, const PxRigidActor* actor, PxHitFlags& queryFlags) override { return U2PCollisionQueryHitType(Capture.GetFilterResult(shape, actor)); }
 	virtual PxQueryHitType::Enum postFilter(const PxFilterData& filterData, const PxQueryHit& hit) override { return PxQueryHitType::eTOUCH; }
 #endif
 
-#if INCLUDE_CHAOS
-	virtual ECollisionQueryHitType PreFilterChaos(const FCollisionFilterData& FilterData, const Chaos::TImplicitObject<float, 3>& Shape, int32 ActorIdx, int32 ShapeIdx) { return Capture.GetFilterResult(Shape, ActorIdx, ShapeIdx); }
-#endif
-
 private:
 	const FSQCapture& Capture;
 };
@@ -537,31 +544,32 @@ FSQCapture::FPhysXSerializerData::~FPhysXSerializerData()
 
 #if INCLUDE_CHAOS && WITH_PHYSX
 
-void PhysXQueryHitToChaosQueryHit(FQueryHit& ChaosHit, const PxQueryHit& PxHit, const FPhysTestSerializer& Serializer)
+void PhysXQueryHitToChaosQueryHit(ChaosInterface::FQueryHit& ChaosHit, const PxQueryHit& PxHit, const FPhysTestSerializer& Serializer)
 {
-	ChaosHit.ActorIdx = Serializer.PhysXActorToChaosIdx(PxHit.actor);
+	ChaosHit.Actor = Serializer.PhysXActorToChaosHandle(PxHit.actor);
 	ChaosHit.Shape = Serializer.PhysXShapeToChaosImplicit(PxHit.shape);
 }
 
-void PhysXLocationHitToChaosLocationHit(FLocationHit& ChaosHit, const PxLocationHit& PxHit, const FPhysTestSerializer& Serializer)
+void PhysXLocationHitToChaosLocationHit(ChaosInterface::FLocationHit& ChaosHit, const PxLocationHit& PxHit, const FPhysTestSerializer& Serializer)
 {
 	PhysXQueryHitToChaosQueryHit(ChaosHit, PxHit, Serializer);
 	ChaosHit.WorldPosition = P2UVector(PxHit.position);
 	ChaosHit.WorldNormal = P2UVector(PxHit.normal);
 	ChaosHit.Flags = P2UHitFlags(PxHit.flags);
+	ChaosHit.Distance = PxHit.distance;
 }
 
-void PhysXHitToChaosHit(FSweepHit& ChaosHit, const PxSweepHit& PxHit, const FPhysTestSerializer& Serializer)
+void PhysXHitToChaosHit(ChaosInterface::FSweepHit& ChaosHit, const PxSweepHit& PxHit, const FPhysTestSerializer& Serializer)
 {
 	PhysXLocationHitToChaosLocationHit(ChaosHit, PxHit, Serializer);
 }
 
-void PhysXHitToChaosHit(FOverlapHit& ChaosHit, const PxOverlapHit& PxHit, const FPhysTestSerializer& Serializer)
+void PhysXHitToChaosHit(ChaosInterface::FOverlapHit& ChaosHit, const PxOverlapHit& PxHit, const FPhysTestSerializer& Serializer)
 {
 	PhysXQueryHitToChaosQueryHit(ChaosHit, PxHit, Serializer);
 }
 
-void PhysXHitToChaosHit(FRaycastHit& ChaosHit, const PxRaycastHit& PxHit, const FPhysTestSerializer& Serializer)
+void PhysXHitToChaosHit(ChaosInterface::FRaycastHit& ChaosHit, const PxRaycastHit& PxHit, const FPhysTestSerializer& Serializer)
 {
 	PhysXLocationHitToChaosLocationHit(ChaosHit, PxHit, Serializer);
 	ChaosHit.U = PxHit.u;
@@ -569,43 +577,47 @@ void PhysXHitToChaosHit(FRaycastHit& ChaosHit, const PxRaycastHit& PxHit, const
 }
 
 template <typename TChaosHit, typename TPhysXHit>
-void PhysXToChaosBufferData(FDynamicHitBuffer<TChaosHit>& ChaosBuffer, const FDynamicHitBuffer<TPhysXHit>& PhysXBuffer, const FPhysTestSerializer& PhysSerializer)
+void PhysXToChaosBufferData(ChaosInterface::FSQHitBuffer<TChaosHit>& ChaosBuffer, const PhysXInterface::FDynamicHitBuffer<TPhysXHit>& PhysXBuffer, const FPhysTestSerializer& PhysSerializer)
 {
-	ChaosBuffer.hasBlock = PhysXBuffer.hasBlock;
-	if (ChaosBuffer.hasBlock)
+	if (PhysXBuffer.hasBlock)
 	{
-		PhysXHitToChaosHit(ChaosBuffer.block, PhysXBuffer.block, PhysSerializer);
+		TChaosHit Hit;
+		PhysXHitToChaosHit(Hit, PhysXBuffer.block, PhysSerializer);
+		ChaosBuffer.SetBlockingHit(Hit);
+
 	}
 	const int32 NumHits = PhysXBuffer.GetNumHits();
-	TArray<TChaosHit> ChaosHits; ChaosHits.AddDefaulted(NumHits);
 	for (int32 Idx = 0; Idx < NumHits; ++Idx)
 	{
-		PhysXHitToChaosHit(ChaosHits[Idx], PhysXBuffer.GetHits()[Idx], PhysSerializer);
+		TChaosHit Hit;
+		PhysXHitToChaosHit(Hit, PhysXBuffer.GetHits()[Idx], PhysSerializer);
+		ChaosBuffer.AddTouchingHit(Hit);
 	}
-
-	ChaosBuffer.processTouches(ChaosHits.GetData(), NumHits);
 }
 
 void FSQCapture::CreateChaosFilterResults()
 {
 #if WITH_PHYSX
-	for (auto Itr : ActorToShapeHitsArray)
+	for (auto Itr : PxActorToShapeHitsArray)
 	{
-		const int32 ActorIdx = PhysSerializer.PhysXActorToChaosIdx(Itr.Key);
-		TArray<ECollisionQueryHitType> FilterResults;
+		Chaos::TGeometryParticle<float, 3>* Actor = PhysSerializer.PhysXActorToChaosHandle(Itr.Key);
+		const auto& ShapesArray = Actor->ShapesArray();
+		TArray<TPair<Chaos::TPerShapeData<float,3>*, ECollisionQueryHitType>> FilterResults;
 
 		const auto& Pairs = Itr.Value;
+		int32 Idx = 0;
 		for (const auto& Pair : Pairs)
 		{
-			FilterResults.Add(Pair.Value);
+			FilterResults.Add(TPair<Chaos::TPerShapeData<float, 3>*, ECollisionQueryHitType>(ShapesArray[Idx++].Get(), Pair.Value));
 		}
-		ChaosActorToHitsArray.Add(ActorIdx, FilterResults);
+		ChaosActorToShapeHitsArray.Add(Actor, FilterResults);
 	}
 #endif
 }
 
 void FSQCapture::CreateChaosData()
 {
+#if INCLUDE_CHAOS
 	using namespace Chaos;
 	check(bDiskDataIsChaos == false);	//we don't support chaos to chaos yet
 
@@ -628,7 +640,6 @@ void FSQCapture::CreateChaosData()
 
 	CreateChaosFilterResults();
 	bChaosDataReady = true;
+#endif
 }
 #endif
-
-#endif //  !WITH_CHAOS_NEEDS_TO_BE_FIXED
diff --git a/Engine/Source/Runtime/PhysicsCore/Public/ChaosInterfaceWrapperCore.h b/Engine/Source/Runtime/PhysicsCore/Public/ChaosInterfaceWrapperCore.h
--- a/Engine/Source/Runtime/PhysicsCore/Public/ChaosInterfaceWrapperCore.h
+++ b/Engine/Source/Runtime/PhysicsCore/Public/ChaosInterfaceWrapperCore.h
@@ -2,26 +2,30 @@
 
 #pragma once
 
-#if WITH_CHAOS
+#if INCLUDE_CHAOS
 
 #include "ChaosSQTypes.h"
-#include "PhysicsInterfaceDeclaresCore.h"
-#include "PhysicsInterfaceTypesCore.h"
-#include "PhysicsInterfaceUtilsCore.h"
 #include "PhysicsInterfaceWrapperShared.h"
 #include "Chaos/ImplicitObject.h"
 #include "Chaos/Capsule.h"
+#include "PhysicsInterfaceTypesCore.h"
 
-template<typename DummyT>
-struct FCallbackDummy {};
+#if WITH_PHYSX
+#include "PhysXPublicCore.h"
+#endif
 
-template <typename T>
-using FPhysicsHitCallback = FCallbackDummy<T>;
+namespace ChaosInterface
+{
+struct FDummyPhysType {};
+struct FDummyPhysActor {};
+
+template<typename DummyT>
+struct FDummyCallback {};
 
 #if WITH_PHYSX
 using FQueryFilterData = PxQueryFilterData;
 #else
-using FQueryFilterData = FPhysTypeDummy;
+using FQueryFilterData = FDummyPhysType;
 #endif
 
 /** We use this struct so that if no conversion is needed in another API, we can avoid the copy (if we think that's critical) */
@@ -65,150 +69,100 @@ struct FPhysicsOverlapInputAdapater
 	FTransform GeomPose;
 };
 
-template<typename HitType>
-class FDynamicHitBuffer : public FCallbackDummy<HitType>
+inline ECollisionShapeType GetImplicitType(const Chaos::TImplicitObject<float, 3>& InGeometry)
 {
-public:
-	FDynamicHitBuffer()
-		: bHasBlockingHit(false)
-	{}
-
-	~FDynamicHitBuffer() {}
-
-	bool ProcessTouchBuffer(TArrayView<const HitType>& Touches)
-	{
-		Hits.Append(Touches);
-		return true;
-	}
-
-	void FinishQuery()
-	{
-		if(bHasBlockingHit)
-		{
-			Hits.Add(CurrentBlockingHit);
-		}
-	}
-
-	bool HasHit() const { return bHasBlockingHit || Hits.Num() == 0; }
-
-	int32 GetNumHits() const
+	switch (InGeometry.GetType())
 	{
-		return Hits.Num();
+	case Chaos::ImplicitObjectType::Sphere: return ECollisionShapeType::Sphere;
+	case Chaos::ImplicitObjectType::Box: return ECollisionShapeType::Box;
+	case Chaos::ImplicitObjectType::Capsule: return ECollisionShapeType::Capsule;
+	case Chaos::ImplicitObjectType::Convex: return ECollisionShapeType::Convex;
+	case Chaos::ImplicitObjectType::TriangleMesh: return ECollisionShapeType::Trimesh;
+	case Chaos::ImplicitObjectType::HeightField: return ECollisionShapeType::Heightfield;
+	default: break;
 	}
 
-	HitType* GetHits()
-	{
-		return (HitType*)Hits.GetData();
-	}
-
-	HitType* GetBlock()
-	{
-		return &CurrentBlockingHit;
-	}
-
-	bool HasBlockingHit() const
-	{
-		return bHasBlockingHit;
-	}
-
-protected:
-	HitType CurrentBlockingHit;
-	bool bHasBlockingHit;
-
-private:
-	static const int32 BufferSize = 512;
-
-	/** Hit buffer used to provide hits via ProcessTouchBuffer */
-	TTypeCompatibleBytes<HitType> HitBuffer[BufferSize];
-
-	/** Hits encountered. Can be larger than HIT_BUFFER_SIZE */
-	TArray<TTypeCompatibleBytes<HitType>, TInlineAllocator<BufferSize>> Hits;
-};
-
-
-inline ECollisionShapeType GetType(const FPhysicsGeometry& InGeometry)
-{
-	return ECollisionShapeType(InGeometry.GetType());
+	return ECollisionShapeType::None;
 }
 
-inline ECollisionShapeType GetGeometryType(const FPhysTypeDummy& Shape)
+inline ECollisionShapeType GetType(const Chaos::TImplicitObject<float, 3>& InGeometry)
 {
-	return ECollisionShapeType::None;
+	return GetImplicitType(InGeometry);
 }
 
-inline float GetRadius(const FPhysicsCapsuleGeometry& InCapsule)
+inline float GetRadius(const Chaos::TCapsule<float>& InCapsule)
 {
 	return InCapsule.GetRadius();
 }
 
-inline float GetHalfHeight(const FPhysicsCapsuleGeometry& InCapsule)
+inline float GetHalfHeight(const Chaos::TCapsule<float>& InCapsule)
 {
 	return InCapsule.GetHeight()/2.;
 }
 
-inline bool HadInitialOverlap(const FPhysTypeDummy& Hit)
+inline bool HadInitialOverlap(const FLocationHit& Hit)
 {
-	return false;
+	return Hit.Distance <= 0.f;
 }
 
-inline Chaos::TImplicitObject<float, 3>* GetShape(const FPhysTypeDummy& Hit)
+inline const Chaos::TPerShapeData<float, 3>* GetShape(const FActorShape& Hit)
 {
-	return nullptr;
+	return Hit.Shape;
 }
 
-inline FPhysActorDummy* GetActor(const FPhysTypeDummy& Hit)
+inline Chaos::TGeometryParticle<float,3>* GetActor(const FActorShape& Hit)
 {
-	return nullptr;
+	return Hit.Actor;
 }
 
-inline float GetDistance(const FPhysTypeDummy& Hit)
+inline float GetDistance(const FLocationHit& Hit)
 {
-	return 0.0f;
+	return Hit.Distance;
 }
 
-inline FVector GetPosition(const FPhysTypeDummy& Hit)
+inline FVector GetPosition(const FLocationHit& Hit)
 {
-	return FVector::ZeroVector;
+	return Hit.WorldPosition;
 }
 
-inline FVector GetNormal(const FPhysTypeDummy& Hit)
+inline FVector GetNormal(const FLocationHit& Hit)
 {
-	return FVector(0.0f, 0.0f, 1.0f);
+	return Hit.WorldNormal;
 }
 
-inline FPhysTypeDummy* GetMaterialFromInternalFaceIndex(const FPhysTypeDummy& Shape, uint32 InternalFaceIndex)
+inline FDummyPhysType* GetMaterialFromInternalFaceIndex(const FDummyPhysType& Shape, uint32 InternalFaceIndex)
 {
 	return nullptr;
 }
 
-inline FHitFlags GetFlags(const FPhysTypeDummy& Hit)
+inline FHitFlags GetFlags(const FLocationHit& Hit)
 {
-	return FHitFlags(EHitFlags::None);
+	return Hit.Flags;
 }
 
-FORCEINLINE void SetFlags(FPhysTypeDummy& Hit, FHitFlags Flags)
+FORCEINLINE void SetFlags(FLocationHit& Hit, FHitFlags Flags)
 {
-	//Hit.flags = U2PHitFlags(Flags);
+	Hit.Flags = Flags;
 }
 
-inline uint32 GetInternalFaceIndex(const FPhysTypeDummy& Hit)
+inline uint32 GetInternalFaceIndex(const FQueryHit& Hit)
 {
-	return 0;
+	return Hit.FaceIndex;
 }
 
-inline void SetInternalFaceIndex(FPhysTypeDummy& Hit, uint32 FaceIndex)
+inline void SetInternalFaceIndex(FQueryHit& Hit, uint32 FaceIndex)
 {
-
+	Hit.FaceIndex = FaceIndex;
 }
 
-inline FCollisionFilterData GetQueryFilterData(const Chaos::TImplicitObject<float, 3>& Shape)
+inline FCollisionFilterData GetQueryFilterData(const Chaos::TPerShapeData<float, 3>& Shape)
 {
-	return FCollisionFilterData();
+	return Shape.QueryData;
 }
 
-inline FCollisionFilterData GetSimulationFilterData(const FPhysTypeDummy& Shape)
+inline FCollisionFilterData GetSimulationFilterData(const Chaos::TPerShapeData<float, 3>& Shape)
 {
-	return FCollisionFilterData();
+	return Shape.QueryData;
 }
 
 inline uint32 GetInvalidPhysicsFaceIndex()
@@ -216,103 +170,52 @@ inline uint32 GetInvalidPhysicsFaceIndex()
 	return 0xffffffff;
 }
 
-inline uint32 GetTriangleMeshExternalFaceIndex(const FPhysTypeDummy& Shape, uint32 InternalFaceIndex)
+inline uint32 GetTriangleMeshExternalFaceIndex(const FDummyPhysType& Shape, uint32 InternalFaceIndex)
 {
 	return GetInvalidPhysicsFaceIndex();
 }
 
-inline FTransform GetGlobalPose(const FPhysActorDummy& RigidActor)
+inline FTransform GetGlobalPose(const FDummyPhysActor& RigidActor)
 {
 	return FTransform::Identity;
 }
 
-inline uint32 GetNumShapes(const FPhysActorDummy& RigidActor)
+inline uint32 GetNumShapes(const FDummyPhysActor& RigidActor)
 {
 	return 0;
 }
 
-inline void GetShapes(const FPhysActorDummy& RigidActor, Chaos::TImplicitObject<float, 3>** ShapesBuffer, uint32 NumShapes)
-{
-
-}
-
-inline void SetActor(FPhysTypeDummy& Hit, FPhysActorDummy* Actor)
-{
-
-}
-
-inline void SetShape(FPhysTypeDummy& Hit, Chaos::TImplicitObject<float, 3>* Shape)
-{
-
-}
-
-template <typename HitType>
-void SetBlock(FPhysicsHitCallback<HitType>& Callback, const HitType& Hit)
+inline void GetShapes(const FDummyPhysActor& RigidActor, Chaos::TImplicitObject<float, 3>** ShapesBuffer, uint32 NumShapes)
 {
 
 }
 
-template <typename HitType>
-void SetHasBlock(FPhysicsHitCallback<HitType>& Callback, bool bHasBlock)
+inline void SetActor(FDummyPhysType& Hit, FDummyPhysActor* Actor)
 {
 
 }
 
-template <typename HitType>
-void ProcessTouches(FPhysicsHitCallback<HitType>& Callback, const TArray<HitType>& TouchingHits)
+inline void SetShape(FDummyPhysType& Hit, Chaos::TImplicitObject<float, 3>* Shape)
 {
 
 }
 
 template <typename HitType>
-void FinalizeQuery(FPhysicsHitCallback<HitType>& Callback)
+HitType* GetBlock(FSQHitBuffer<HitType>& Callback)
 {
-
+	return Callback.GetBlock();
 }
 
 template <typename HitType>
-HitType* GetBlock(const FPhysicsHitCallback<HitType>& Callback)
+bool GetHasBlock(const FSQHitBuffer<HitType>& Callback)
 {
-	return nullptr;
+	return Callback.HasBlockingHit();
 }
 
-template <typename HitType>
-bool GetHasBlock(const FPhysicsHitCallback<HitType>& Callback)
-{
-	return false;
-}
+} // namespace ChaosInterface
 
-template <typename HitType>
-FORCEINLINE_DEBUGGABLE bool Insert(FCallbackDummy<HitType>& Callback, const HitType& Hit, bool bBlocking)
-{
-	// Previous Physx implementation
-	//if(!Callback.hasBlock || GetDistance(Hit) < GetDistance(Callback.block))
-	//{
-	//	if(bBlocking)
-	//	{
-	//		Callback.block = Hit;
-	//		Callback.hasBlock = true;
-	//	}
-	//	else
-	//	{
-	//		if(Callback.maxNbTouches > 0)
-	//		{
-	//			Callback.processTouches(&Hit, 1);
-	//
-	//		}
-	//	}
-	//}
-
-	return true;
-}
-
-template <typename HitType>
-FORCEINLINE bool InsertOverlap(FCallbackDummy<HitType>& Callback, const HitType& Hit)
-{
-	// Previous Physx implementation
-	//return Callback.processTouches(&Hit, 1);
-
-	return true;
-}
+#endif // WITH_CHAOS
 
-#endif
\ No newline at end of file
+#if WITH_CHAOS && (!defined(PHYSICS_INTERFACE_PHYSX) || !PHYSICS_INTERFACE_PHYSX)
+using namespace ChaosInterface;
+#endif
diff --git a/Engine/Source/Runtime/PhysicsCore/Public/ChaosSQTypes.h b/Engine/Source/Runtime/PhysicsCore/Public/ChaosSQTypes.h
--- a/Engine/Source/Runtime/PhysicsCore/Public/ChaosSQTypes.h
+++ b/Engine/Source/Runtime/PhysicsCore/Public/ChaosSQTypes.h
@@ -2,56 +2,168 @@
 
 #pragma once 
 
+#if INCLUDE_CHAOS
+
 #include "CoreMinimal.h"
 #include "PhysicsInterfaceWrapperShared.h"
+#include "Chaos/ParticleHandleFwd.h"
+#include "Chaos/Declares.h"
 
 namespace Chaos
 {
 	template <typename T, int d>
-	class TImplicitObject;
+	class TPerShapeData;
 }
 
-struct FActorShape
+namespace ChaosInterface
 {
-	int32 ActorIdx;	//todo: this needs to be a pointer or handle once we decide other parts of chaos interface
-	const Chaos::TImplicitObject<float,3>* Shape;
-};
 
-struct FQueryHit : public FActorShape
-{
-	FQueryHit() : FaceIndex(-1) {}
-	
-	/**
-	Face index of touched triangle, for triangle meshes, convex meshes and height fields. Defaults to -1 if face index is not available
+	struct FActorShape
+	{
+		Chaos::TGeometryParticle<float, 3>* Actor;
+		const Chaos::TPerShapeData<float, 3>* Shape;
+	};
+
+	struct FQueryHit : public FActorShape
+	{
+		FQueryHit() : FaceIndex(-1) {}
+
+		/**
+		Face index of touched triangle, for triangle meshes, convex meshes and height fields. Defaults to -1 if face index is not available
+		*/
+
+		uint32 FaceIndex;
+	};
+
+	struct FLocationHit : public FQueryHit
+	{
+		FHitFlags Flags;
+		FVector WorldPosition;
+		FVector WorldNormal;
+		float Distance;
+	};
+
+	struct FRaycastHit : public FLocationHit
+	{
+		float U;
+		float V;
+	};
+
+	struct FOverlapHit : public FQueryHit
+	{
+	};
+
+	struct FSweepHit : public FLocationHit
+	{
+	};
+
+	/** Stores the results of scene queries. This can be passed around to multiple SQAccelerators and is responsible for sorting the results and pruning based on blocking.
+	IncFlushCount / DecFlushCount is used to ensure any final sort / pruning operation is done when all SQAccelerators are finished. If you are passing this into multiple accelerators
+	you should call IncFlushCount / DecFlushCount yourself as otherwise each accelerator will trigger its own sort / prune.
 	*/
+	template<typename HitType>
+	class FSQHitBuffer
+	{
+	public:
+		FSQHitBuffer(bool bSingle = false)
+			: AcceleratorDepth(0)
+			, bHasBlockingHit(false)
+			, bSingleResult(bSingle)
+		{
+			if (!bSingle)
+			{
+				Hits.Reserve(512);
+			}
+		}
 
-	uint32 FaceIndex;
-};
+		virtual ~FSQHitBuffer() {}
 
-struct FLocationHit : public FQueryHit
-{
-	FHitFlags Flags;
-	FVector WorldPosition;
-	FVector WorldNormal;
-	float Distance;
-};
+		//Called 
+		void IncFlushCount()
+		{
+			++AcceleratorDepth;
+		}
 
-struct FRaycastHit : public FLocationHit
-{
-	float U;
-	float V;
-};
+		void DecFlushCount()
+		{
+			--AcceleratorDepth;
+			if (AcceleratorDepth == 0)
+			{
+				FinishQuery();
+			}
+		}
 
-struct FOverlapHit : public FQueryHit
-{
-};
+		bool HasHit() const { return GetNumHits(); }
+		int32 GetNumHits() const { return Hits.Num(); }
+		HitType* GetHits() { return Hits.GetData(); }
+		const HitType* GetHits() const { return Hits.GetData(); }
 
-struct FSweepHit : public FLocationHit
-{
-};
+		HitType* GetBlock() { return HasBlockingHit() ? &Hits.Last() : nullptr; }
+		const HitType* GetBlock() const { return HasBlockingHit() ? &Hits.Last() : nullptr; }
 
+		bool HasBlockingHit() const { return bHasBlockingHit; }
 
-FORCEINLINE float GetDistance(const FLocationHit& Hit)
-{
-	return Hit.Distance;
-}
\ No newline at end of file
+		void SetBlockingHit(const HitType& InBlockingHit)
+		{
+			CurrentBlockingHit = InBlockingHit;
+			bHasBlockingHit = true;
+		}
+
+		void AddTouchingHit(const HitType& InTouchingHit)
+		{
+			Hits.Add(InTouchingHit);
+		}
+
+		/** Does not do any distance verification. This is up to the SQ code to manage */
+		void InsertHit(const HitType& Hit, bool bBlocking)
+		{
+			if (bBlocking)
+			{
+				SetBlockingHit(Hit);
+			}
+			else
+			{
+				AddTouchingHit(Hit);
+			}
+		}
+
+		bool WantsSingleResult() const { return bSingleResult; }
+
+	private:
+
+		void FinishQuery()
+		{
+			//todo: sort
+			//Hits.Sort();
+			if (bHasBlockingHit)
+			{
+				//remove any touches that are past the blocking hit
+				/*for (int32 HitIdx = Hits.Num() - 1; HitIdx >= 0; --HitIdx)
+				{
+					if(Hits[HitIdx])
+				}*/
+				//always put blocking hit in the back
+				Hits.Add(CurrentBlockingHit);
+			}
+		}
+
+		HitType CurrentBlockingHit;
+		int32 AcceleratorDepth;
+		bool bHasBlockingHit;
+		bool bSingleResult;
+
+		TArray<HitType> Hits;
+	};
+
+	template<typename HitType>
+	class FSQSingleHitBuffer : public FSQHitBuffer<HitType>
+	{
+	public:
+		FSQSingleHitBuffer()
+			: FSQHitBuffer<HitType>(/*bSingle=*/ true)
+		{}
+	};
+
+}
+
+#endif
diff --git a/Engine/Source/Runtime/PhysicsCore/Public/CollisionQueryFilterCallbackCore.h b/Engine/Source/Runtime/PhysicsCore/Public/CollisionQueryFilterCallbackCore.h
--- a/Engine/Source/Runtime/PhysicsCore/Public/CollisionQueryFilterCallbackCore.h
+++ b/Engine/Source/Runtime/PhysicsCore/Public/CollisionQueryFilterCallbackCore.h
@@ -6,8 +6,6 @@
 #include "PhysicsInterfaceTypesCore.h"
 #if WITH_PHYSX
 #include "PhysXPublicCore.h"
-#endif
-#if PHYSICS_INTERFACE_PHYSX
 #include "PhysXInterfaceWrapperCore.h"
 #endif
 #if WITH_CHAOS
@@ -42,10 +40,14 @@ class ICollisionQueryFilterCallbackBase
 public:
 	
 	virtual ~ICollisionQueryFilterCallbackBase() {}
-	virtual ECollisionQueryHitType PostFilter(const FCollisionFilterData& FilterData, const FPhysicsQueryHit& Hit) = 0;
-	virtual ECollisionQueryHitType PreFilter(const FCollisionFilterData& FilterData, const FPhysicsShape& Shape, const FPhysicsActor& Actor) = 0;
 #if INCLUDE_CHAOS
-	virtual ECollisionQueryHitType PreFilterChaos(const FCollisionFilterData& FilterData, const Chaos::TImplicitObject<float, 3>& Shape, int32 ActorIdx, int32 ShapeIdx) = 0;
+	virtual ECollisionQueryHitType PostFilter(const FCollisionFilterData& FilterData, const ChaosInterface::FQueryHit& Hit) = 0;
+	virtual ECollisionQueryHitType PreFilter(const FCollisionFilterData& FilterData, const Chaos::TPerShapeData<float, 3>& Shape, const Chaos::TGeometryParticle<float, 3>& Actor) = 0;
+#endif
+
+#if WITH_PHYSX
+	virtual ECollisionQueryHitType PostFilter(const FCollisionFilterData& FilterData, const physx::PxQueryHit& Hit) = 0;
+	virtual ECollisionQueryHitType PreFilter(const FCollisionFilterData& FilterData, const physx::PxShape& Shape, physx::PxRigidActor& Actor) = 0;
 #endif
 };
 
@@ -53,16 +55,33 @@ class PHYSICSCORE_API FBlockAllQueryCallback : public ICollisionQueryFilterCallb
 {
 public:
 	virtual ~FBlockAllQueryCallback() {}
-	virtual ECollisionQueryHitType PostFilter(const FCollisionFilterData& FilterData, const FPhysicsQueryHit& Hit) override { return ECollisionQueryHitType::Block; }
-	virtual ECollisionQueryHitType PreFilter(const FCollisionFilterData& FilterData, const FPhysicsShape& Shape, const FPhysicsActor& Actor) override { return ECollisionQueryHitType::Block; }
+#if INCLUDE_CHAOS
+	virtual ECollisionQueryHitType PostFilter(const FCollisionFilterData& FilterData, const ChaosInterface::FQueryHit& Hit) override { return ECollisionQueryHitType::Block; }
+	virtual ECollisionQueryHitType PreFilter(const FCollisionFilterData& FilterData, const Chaos::TPerShapeData<float, 3>& Shape, const Chaos::TGeometryParticle<float, 3>& Actor) override { return ECollisionQueryHitType::Block; }
+#endif
 
 #if WITH_PHYSX
+	virtual ECollisionQueryHitType PostFilter(const FCollisionFilterData& FilterData, const physx::PxQueryHit& Hit) override { return ECollisionQueryHitType::Block; }
+	virtual ECollisionQueryHitType PreFilter(const FCollisionFilterData& FilterData, const physx::PxShape& Shape, physx::PxRigidActor& Actor) override { return ECollisionQueryHitType::Block; }
 	virtual PxQueryHitType::Enum preFilter(const PxFilterData& filterData, const PxShape* shape, const PxRigidActor* actor, PxHitFlags& queryFlags) override { return PxQueryHitType::eBLOCK; }
 	virtual PxQueryHitType::Enum postFilter(const PxFilterData& filterData, const PxQueryHit& hit) override { return PxQueryHitType::eBLOCK; }
 #endif
+};
 
+class PHYSICSCORE_API FOverlapAllQueryCallback : public ICollisionQueryFilterCallbackBase
+{
+public:
+	virtual ~FOverlapAllQueryCallback() {}
 #if INCLUDE_CHAOS
-	virtual ECollisionQueryHitType PreFilterChaos(const FCollisionFilterData& FilterData, const Chaos::TImplicitObject<float, 3>& Shape, int32 ActorIdx, int32 ShapeIdx) { return ECollisionQueryHitType::Block; }
+	virtual ECollisionQueryHitType PostFilter(const FCollisionFilterData& FilterData, const ChaosInterface::FQueryHit& Hit) override { return ECollisionQueryHitType::Touch; }
+	virtual ECollisionQueryHitType PreFilter(const FCollisionFilterData& FilterData, const Chaos::TPerShapeData<float, 3>& Shape, const Chaos::TGeometryParticle<float, 3>& Actor) override { return ECollisionQueryHitType::Touch; }
+#endif
+
+#if WITH_PHYSX
+	virtual ECollisionQueryHitType PostFilter(const FCollisionFilterData& FilterData, const physx::PxQueryHit& Hit) override { return ECollisionQueryHitType::Touch; }
+	virtual ECollisionQueryHitType PreFilter(const FCollisionFilterData& FilterData, const physx::PxShape& Shape, physx::PxRigidActor& Actor) override { return ECollisionQueryHitType::Touch; }
+	virtual PxQueryHitType::Enum preFilter(const PxFilterData& filterData, const PxShape* shape, const PxRigidActor* actor, PxHitFlags& queryFlags) override { return PxQueryHitType::eTOUCH; }
+	virtual PxQueryHitType::Enum postFilter(const PxFilterData& filterData, const PxQueryHit& hit) override { return PxQueryHitType::eTOUCH; }
 #endif
 };
 
diff --git a/Engine/Source/Runtime/PhysicsCore/Public/PhysTestSerializer.h b/Engine/Source/Runtime/PhysicsCore/Public/PhysTestSerializer.h
--- a/Engine/Source/Runtime/PhysicsCore/Public/PhysTestSerializer.h
+++ b/Engine/Source/Runtime/PhysicsCore/Public/PhysTestSerializer.h
@@ -5,11 +5,9 @@
 #include "CoreMinimal.h"
 #include "PhysicsInterfaceDeclaresCore.h"
 
-#if  !WITH_CHAOS_NEEDS_TO_BE_FIXED
-
-
 #if INCLUDE_CHAOS
-#include "Chaos/PBDRigidsEvolution.h"
+#include "Chaos/PBDRigidsEvolutionGBF.h"
+#include "Chaos/ParticleHandle.h"
 #endif
 #include "SQCapture.h"
 
@@ -73,7 +71,9 @@ public:
 
 #if INCLUDE_CHAOS
 			GetChaosData();
+#if WITH_PHYSX
 			SQCapture->CreateChaosData();
+#endif
 #endif
 		}
 		return SQCapture.Get();
@@ -105,8 +105,8 @@ public:
 #endif
 
 #if WITH_PHYSX && INCLUDE_CHAOS
-	int32 PhysXActorToChaosIdx(physx::PxActor* Actor) const { return PxActorToChaosIdx.FindChecked(Actor); }
-	Chaos::TImplicitObject<float,3>* PhysXShapeToChaosImplicit(physx::PxShape* Shape) const { return PxShapeToChaosImplicit.FindRef(Shape); }
+	Chaos::TGeometryParticle<float,3>* PhysXActorToChaosHandle(physx::PxActor* Actor) const { return PxActorToChaosHandle.FindChecked(Actor)->GTGeometryParticle(); }
+	Chaos::TPerShapeData<float,3>* PhysXShapeToChaosImplicit(physx::PxShape* Shape) const { return PxShapeToChaosShapes.FindRef(Shape); }
 #endif
 
 private:
@@ -124,6 +124,8 @@ private:
 
 #if INCLUDE_CHAOS
 	TUniquePtr<Chaos::TPBDRigidsEvolutionGBF<float, 3>> ChaosEvolution;
+	Chaos::TPBDRigidsSOAs<float, 3> Particles;
+	TArray <TUniquePtr<Chaos::TGeometryParticle<float, 3>>> GTParticles;
 #endif
 
 #if WITH_PHYSX
@@ -145,11 +147,9 @@ private:
 #endif
 
 #if WITH_PHYSX && INCLUDE_CHAOS
-	TMap<physx::PxActor*, int32> PxActorToChaosIdx;
-	TMap<physx::PxShape*, Chaos::TImplicitObject<float, 3>*> PxShapeToChaosImplicit;
+	TMap<physx::PxActor*, Chaos::TGeometryParticleHandle<float, 3>*> PxActorToChaosHandle;
+	TMap<physx::PxShape*, Chaos::TPerShapeData<float, 3>*> PxShapeToChaosShapes;
 #endif
 };
 
-#endif
-
-#endif
+#endif
\ No newline at end of file
