Reorder #if to satisfy include tool
#rnx

diff --git a/Engine/Source/Runtime/PhysicsCore/Public/PhysicsInterfaceDeclaresCore.h b/Engine/Source/Runtime/PhysicsCore/Public/PhysicsInterfaceDeclaresCore.h
--- a/Engine/Source/Runtime/PhysicsCore/Public/PhysicsInterfaceDeclaresCore.h
+++ b/Engine/Source/Runtime/PhysicsCore/Public/PhysicsInterfaceDeclaresCore.h
@@ -14,56 +14,7 @@
 #define WITH_IMMEDIATE_PHYSX 0
 #endif
 
-#if PHYSICS_INTERFACE_PHYSX
+#if WITH_CHAOS
 
-namespace physx
-{
-	struct PxLocationHit;
-	struct PxSweepHit;
-	struct PxRaycastHit;
-	struct PxOverlapHit;
-	struct PxQueryHit;
-
-	class PxTransform;
-	class PxShape;
-	class PxGeometry;
-	class PxCapsuleGeometry;
-	class PxMaterial;
-	class PxRigidActor;
-
-	template<typename T>
-	struct PxHitBuffer;
-	
-	template<typename T>
-	struct PxHitCallback;
-}
-
-namespace PhysXInterface
-{
-	struct FDummyPhysType {};
-}
-
-using FHitLocation = physx::PxLocationHit;
-using FHitSweep = physx::PxSweepHit;
-using FHitRaycast = physx::PxRaycastHit;
-using FHitOverlap = physx::PxOverlapHit;
-using FPhysicsQueryHit = physx::PxQueryHit;
-
-using FPhysicsTransform = physx::PxTransform;
-using FPhysTypeDummy = PhysXInterface::FDummyPhysType;
-
-using FPhysicsShape = physx::PxShape;
-using FPhysicsGeometry = physx::PxGeometry;
-using FPhysicsCapsuleGeometry = physx::PxCapsuleGeometry;
-using FPhysicsMaterial = physx::PxMaterial;
-using FPhysicsActor = physx::PxRigidActor;
-
-template <typename T>
-using FPhysicsHitCallback = physx::PxHitCallback<T>;
-
-struct FQueryDebugParams {};
-
-#elif WITH_CHAOS
-
 #include "ChaosSQTypes.h"
 
@@ -128,6 +79,55 @@
 template <typename T>
 using FDynamicHitBuffer = ChaosInterface::FSQHitBuffer<T>;
 
+#elif PHYSICS_INTERFACE_PHYSX
+
+namespace physx
+{
+	struct PxLocationHit;
+	struct PxSweepHit;
+	struct PxRaycastHit;
+	struct PxOverlapHit;
+	struct PxQueryHit;
+
+	class PxTransform;
+	class PxShape;
+	class PxGeometry;
+	class PxCapsuleGeometry;
+	class PxMaterial;
+	class PxRigidActor;
+
+	template<typename T>
+	struct PxHitBuffer;
+	
+	template<typename T>
+	struct PxHitCallback;
+}
+
+namespace PhysXInterface
+{
+	struct FDummyPhysType {};
+}
+
+using FHitLocation = physx::PxLocationHit;
+using FHitSweep = physx::PxSweepHit;
+using FHitRaycast = physx::PxRaycastHit;
+using FHitOverlap = physx::PxOverlapHit;
+using FPhysicsQueryHit = physx::PxQueryHit;
+
+using FPhysicsTransform = physx::PxTransform;
+using FPhysTypeDummy = PhysXInterface::FDummyPhysType;
+
+using FPhysicsShape = physx::PxShape;
+using FPhysicsGeometry = physx::PxGeometry;
+using FPhysicsCapsuleGeometry = physx::PxCapsuleGeometry;
+using FPhysicsMaterial = physx::PxMaterial;
+using FPhysicsActor = physx::PxRigidActor;
+
+template <typename T>
+using FPhysicsHitCallback = physx::PxHitCallback<T>;
+
+struct FQueryDebugParams {};
+
 #else
 
 static_assert(false, "A physics engine interface must be defined to build");
