Chaos: Found typo bug in interface

diff --git a/Engine/Source/Runtime/PhysicsCore/Private/ChaosInterfaceWrapperCore.cpp b/Engine/Source/Runtime/PhysicsCore/Private/ChaosInterfaceWrapperCore.cpp
--- a/Engine/Source/Runtime/PhysicsCore/Private/ChaosInterfaceWrapperCore.cpp
+++ b/Engine/Source/Runtime/PhysicsCore/Private/ChaosInterfaceWrapperCore.cpp
@@ -42,7 +42,7 @@ namespace ChaosInterface
 
 	FCollisionFilterData GetSimulationFilterData(const Chaos::TPerShapeData<float, 3>& Shape)
 	{
-		return Shape.QueryData;
+		return Shape.SimData;
 	}
 
 
