Fix bubble shield crash, GetImplicitType failing on heightfields, and heightfield material/face index not being returned from SQ.

diff --git a/Engine/Source/Runtime/Engine/Private/PhysicsEngine/Experimental/PhysInterface_Chaos.cpp b/Engine/Source/Runtime/Engine/Private/PhysicsEngine/Experimental/PhysInterface_Chaos.cpp
--- a/Engine/Source/Runtime/Engine/Private/PhysicsEngine/Experimental/PhysInterface_Chaos.cpp
+++ b/Engine/Source/Runtime/Engine/Private/PhysicsEngine/Experimental/PhysInterface_Chaos.cpp
@@ -1946,7 +1946,7 @@ bool FPhysInterface_Chaos::GetSquaredDistanceToBody(const FBodyInstance* InInsta
 
 			ECollisionShapeType GeomType = FPhysicsInterface::GetShapeType(Shape);
 
-			if (GeomType == ECollisionShapeType::Trimesh)
+			if (!Shape.GetGeometry().IsConvex())
 			{
 				// Type unsupported for this function, but some other shapes will probably work. 
 				continue;
diff --git a/Engine/Source/Runtime/PhysicsCore/Private/ChaosInterfaceWrapperCore.cpp b/Engine/Source/Runtime/PhysicsCore/Private/ChaosInterfaceWrapperCore.cpp
--- a/Engine/Source/Runtime/PhysicsCore/Private/ChaosInterfaceWrapperCore.cpp
+++ b/Engine/Source/Runtime/PhysicsCore/Private/ChaosInterfaceWrapperCore.cpp
@@ -4,14 +4,15 @@
 
 #include "Chaos/Capsule.h"
 #include "Chaos/ImplicitObject.h"
+#include "Chaos/ImplicitObjectTransformed.h"
 #include "Chaos/ParticleHandle.h"
 #include "PhysXPublicCore.h"
 
 namespace ChaosInterface
 {
-	ECollisionShapeType GetImplicitType(const Chaos::FImplicitObject& InGeometry)
+	FORCEINLINE ECollisionShapeType ImplicitTypeToCollisionType(int32 ImplicitObjectType)
 	{
-		switch (Chaos::GetInnerType(InGeometry.GetType()))
+		switch (ImplicitObjectType)
 		{
 		case Chaos::ImplicitObjectType::Sphere: return ECollisionShapeType::Sphere;
 		case Chaos::ImplicitObjectType::Box: return ECollisionShapeType::Box;
@@ -25,6 +26,20 @@ namespace ChaosInterface
 		return ECollisionShapeType::None;
 	}
 
+
+	ECollisionShapeType GetImplicitType(const Chaos::FImplicitObject& InGeometry)
+	{
+		using namespace Chaos;
+		int32 ImplicitObjectType = GetInnerType(InGeometry.GetType());
+
+		if (ImplicitObjectType == ImplicitObjectType::Transformed)
+		{
+			ImplicitObjectType = static_cast<const TImplicitObjectTransformed<FReal, 3>*>(&InGeometry)->Object()->GetType();
+		}
+
+		return ImplicitTypeToCollisionType(ImplicitObjectType);
+	}
+
 	float GetRadius(const Chaos::TCapsule<float>& InCapsule)
 	{
 		return InCapsule.GetRadius();
