#chaos - Fix FindFaceIndex() not transforming world direction to local space when querying Geometry->FindMostOpposingFace().

diff --git a/Engine/Source/Runtime/PhysicsCore/Private/PhysicsInterfaceUtilsCore.cpp b/Engine/Source/Runtime/PhysicsCore/Private/PhysicsInterfaceUtilsCore.cpp
--- a/Engine/Source/Runtime/PhysicsCore/Private/PhysicsInterfaceUtilsCore.cpp
+++ b/Engine/Source/Runtime/PhysicsCore/Private/PhysicsInterfaceUtilsCore.cpp
@@ -269,6 +269,7 @@ uint32 FindFaceIndex(const FHitLocation& PHit, const FVector& UnitDir)
 #else
 	const FTransform WorldTM(PHit.Actor->R(), PHit.Actor->X());
 	const FVector LocalPosition = WorldTM.InverseTransformPositionNoScale(PHit.WorldPosition);
-	return PHit.Shape->Geometry->FindMostOpposingFace(LocalPosition, UnitDir, PHit.FaceIndex, 1);	//todo:this number matches the one above, but is it right?
+	const FVector LocalNormal = WorldTM.InverseTransformVectorNoScale(UnitDir);
+	return PHit.Shape->Geometry->FindMostOpposingFace(LocalPosition, LocalNormal, PHit.FaceIndex, 1);	//todo:this number matches the one above, but is it right?
 #endif
 }
