PhysX strip out of 12.10

diff --git a/Engine/Source/Runtime/Engine/Classes/PhysicalMaterials/PhysicalMaterial.h b/Engine/Source/Runtime/Engine/Classes/PhysicalMaterials/PhysicalMaterial.h
--- a/Engine/Source/Runtime/Engine/Classes/PhysicalMaterials/PhysicalMaterial.h
+++ b/Engine/Source/Runtime/Engine/Classes/PhysicalMaterials/PhysicalMaterial.h
@@ -118,9 +118,7 @@ public:
 
 	FPhysicsMaterialHandle MaterialHandle;
 
-#if WITH_PHYSX
 	FPhysxUserData PhysxUserData;
-#endif
 
 	//~ Begin UObject Interface
 #if WITH_EDITOR
diff --git a/Engine/Source/Runtime/Engine/Classes/PhysicsEngine/ShapeElem.h b/Engine/Source/Runtime/Engine/Classes/PhysicsEngine/ShapeElem.h
--- a/Engine/Source/Runtime/Engine/Classes/PhysicsEngine/ShapeElem.h
+++ b/Engine/Source/Runtime/Engine/Classes/PhysicsEngine/ShapeElem.h
@@ -32,18 +32,14 @@ struct FKShapeElem
 	: RestOffset(0.f)
 	, ShapeType(EAggCollisionShape::Unknown)
 	, bContributeToMass(true)
-#if WITH_PHYSX
 	, UserData(this)
-#endif
 	{}
 
 	FKShapeElem(EAggCollisionShape::Type InShapeType)
 	: RestOffset(0.f)
 	, ShapeType(InShapeType)
 	, bContributeToMass(true)
-#if WITH_PHYSX
 	, UserData(this)
-#endif
 	{}
 
 	FKShapeElem(const FKShapeElem& Copy)
@@ -51,9 +47,7 @@ struct FKShapeElem
 	, Name(Copy.Name)
 	, ShapeType(Copy.ShapeType)
 	, bContributeToMass(Copy.bContributeToMass)
-#if WITH_PHYSX
 	, UserData(this)
-#endif
 	{
 	}
 
@@ -72,9 +66,7 @@ struct FKShapeElem
 		return (T*)this;
 	}
 
-#if WITH_PHYSX
 	const FPhysxUserData* GetUserData() const { FPhysxUserData::Set<FKShapeElem>((void*)&UserData, const_cast<FKShapeElem*>(this));  return &UserData; }
-#endif // WITH_PHYSX
 
 	ENGINE_API static EAggCollisionShape::Type StaticShapeType;
 
@@ -118,7 +110,5 @@ private:
 	UPROPERTY(Category=Shape, EditAnywhere)
 	uint8 bContributeToMass : 1;
 
-#if WITH_PHYSX
 	FPhysxUserData UserData;
-#endif // WITH_PHYSX
 };
diff --git a/Engine/Source/Runtime/Engine/Private/Collision/WorldCollision.cpp b/Engine/Source/Runtime/Engine/Private/Collision/WorldCollision.cpp
--- a/Engine/Source/Runtime/Engine/Private/Collision/WorldCollision.cpp
+++ b/Engine/Source/Runtime/Engine/Private/Collision/WorldCollision.cpp
@@ -392,7 +392,6 @@ bool UWorld::ComponentSweepMulti(TArray<struct FHitResult>& OutHits, class UPrim
 	SCOPE_CYCLE_COUNTER(STAT_Collision_GeomSweepMultiple);
 	bool bHaveBlockingHit = false;
 
-#if WITH_PHYSX
 	FPhysicsCommand::ExecuteRead(BodyInstance->ActorHandle, [&](const FPhysicsActorHandle& Actor)
 	{
 		if(!FPhysicsInterface::IsValid(Actor))
@@ -444,7 +443,6 @@ bool UWorld::ComponentSweepMulti(TArray<struct FHitResult>& OutHits, class UPrim
 			OutHits.Append(TmpHits);	//todo: should these be made unique?
 		}
 	});
-#endif //WITH_PHYSX
 
 	return bHaveBlockingHit;
 }
diff --git a/Engine/Source/Runtime/Engine/Private/Components/PostProcessComponent.cpp b/Engine/Source/Runtime/Engine/Private/Components/PostProcessComponent.cpp
--- a/Engine/Source/Runtime/Engine/Private/Components/PostProcessComponent.cpp
+++ b/Engine/Source/Runtime/Engine/Private/Components/PostProcessComponent.cpp
@@ -20,7 +20,6 @@ bool UPostProcessComponent::EncompassesPoint(FVector Point, float SphereRadius/*
 	{
 		float Distance = -1.f;
 
-#if WITH_PHYSX
 		FVector ClosestPoint;
 		float DistanceSq = -1.f;
 
@@ -29,7 +28,6 @@ bool UPostProcessComponent::EncompassesPoint(FVector Point, float SphereRadius/*
 			Distance = FMath::Sqrt(DistanceSq);
 		}
 		else
-#endif
 		{
 			FBoxSphereBounds SphereBounds = ParentShape->CalcBounds(ParentShape->GetComponentTransform());	
 			if (ParentShape->IsA<USphereComponent>())
diff --git a/Engine/Source/Runtime/Engine/Private/Components/ShapeComponent.cpp b/Engine/Source/Runtime/Engine/Private/Components/ShapeComponent.cpp
--- a/Engine/Source/Runtime/Engine/Private/Components/ShapeComponent.cpp
+++ b/Engine/Source/Runtime/Engine/Private/Components/ShapeComponent.cpp
@@ -100,7 +100,6 @@ template <> void UShapeComponent::AddShapeToGeomArray<FKBoxElem>() { ShapeBodySe
 template <> void UShapeComponent::AddShapeToGeomArray<FKSphereElem>() { ShapeBodySetup->AggGeom.SphereElems.Add(FKSphereElem()); }
 template <> void UShapeComponent::AddShapeToGeomArray<FKSphylElem>() { ShapeBodySetup->AggGeom.SphylElems.Add(FKSphylElem()); }
 
-#if WITH_PHYSX
 template <>
 void UShapeComponent::SetShapeToNewGeom<FKBoxElem>(const FPhysicsShapeHandle& Shape)
 {
@@ -118,7 +117,6 @@ void UShapeComponent::SetShapeToNewGeom<FKSphylElem>(const FPhysicsShapeHandle&
 {
 	FPhysicsInterface::SetUserData(Shape, (void*)ShapeBodySetup->AggGeom.SphylElems[0].GetUserData());
 }
-#endif
 
 template <typename ShapeElemType>
 void UShapeComponent::CreateShapeBodySetupIfNeeded()
@@ -150,7 +148,6 @@ void UShapeComponent::CreateShapeBodySetupIfNeeded()
 		{
 			if(BodyInstance.IsValidBodyInstance())
 			{
-#if WITH_PHYSX
 				FPhysicsCommand::ExecuteWrite(BodyInstance.GetActorReferenceWithWelding(), [this](const FPhysicsActorHandle& Actor)
 				{
 					TArray<FPhysicsShapeHandle> Shapes;
@@ -165,7 +162,6 @@ void UShapeComponent::CreateShapeBodySetupIfNeeded()
 						}
 					}
 				});
-#endif
 			}
 		}
 	}
diff --git a/Engine/Source/Runtime/Engine/Private/Components/SplineMeshComponent.cpp b/Engine/Source/Runtime/Engine/Private/Components/SplineMeshComponent.cpp
--- a/Engine/Source/Runtime/Engine/Private/Components/SplineMeshComponent.cpp
+++ b/Engine/Source/Runtime/Engine/Private/Components/SplineMeshComponent.cpp
@@ -913,7 +913,15 @@ UBodySetup* USplineMeshComponent::GetBodySetup()
 	{
 		return BodySetup;
 	}
-#endif // WITH_PHYSX
+#elif WITH_CHAOS
+	// Don't return a body setup that has no collision, it means we are interactively moving the spline and don't want to build collision.
+	// Instead we explicitly build collision with USplineMeshComponent::RecreateCollision()
+	if (BodySetup != NULL && (BodySetup->ChaosTriMeshes.Num() || BodySetup->AggGeom.GetElementCount() > 0))
+	{
+		return BodySetup;
+	}
+#endif // WITH_CHAOS
+
 	return NULL;
 }
 
diff --git a/Engine/Source/Runtime/Engine/Private/PhysicsEngine/BodyInstance.cpp b/Engine/Source/Runtime/Engine/Private/PhysicsEngine/BodyInstance.cpp
--- a/Engine/Source/Runtime/Engine/Private/PhysicsEngine/BodyInstance.cpp
+++ b/Engine/Source/Runtime/Engine/Private/PhysicsEngine/BodyInstance.cpp
@@ -2926,7 +2926,69 @@ float gPerCm3ToKgPerCm3(float gPerCm3)
 	return gPerCm3 * gToKG;
 }
 
-#if WITH_PHYSX
+#if WITH_CHAOS
+/** Computes and adds the mass properties (inertia, com, etc...) based on the mass settings of the body instance. */
+Chaos::TMassProperties<float, 3> ComputeMassProperties(const FBodyInstance* OwningBodyInstance, TArray<FPhysicsShapeHandle> Shapes, const FTransform& MassModifierTransform)
+{
+	// physical material - nothing can weigh less than hydrogen (0.09 kg/m^3)
+	float DensityKGPerCubicUU = 1.0f;
+	float RaiseMassToPower = 0.75f;
+	if (UPhysicalMaterial* PhysMat = OwningBodyInstance->GetSimplePhysicalMaterial())
+	{
+		DensityKGPerCubicUU = FMath::Max(KgPerM3ToKgPerCm3(0.09f), gPerCm3ToKgPerCm3(PhysMat->Density));
+		RaiseMassToPower = PhysMat->RaiseMassToPower;
+	}
+
+	Chaos::TMassProperties<float, 3> MassProps;
+	FPhysicsInterface::CalculateTMassPropertiesFromShapeCollection(MassProps, Shapes, DensityKGPerCubicUU);
+
+	float OldMass = MassProps.Mass;
+	float NewMass = 0.f;
+
+	if (OwningBodyInstance->bOverrideMass == false)
+	{
+		float UsePow = FMath::Clamp<float>(RaiseMassToPower, KINDA_SMALL_NUMBER, 1.f);
+		NewMass = FMath::Pow(OldMass, UsePow);
+
+		// Apply user-defined mass scaling.
+		NewMass = FMath::Max(OwningBodyInstance->MassScale * NewMass, 0.001f);	//min weight of 1g
+	}
+	else
+	{
+		NewMass = FMath::Max(OwningBodyInstance->GetMassOverride(), 0.001f);	//min weight of 1g
+	}
+
+	check(NewMass > 0.f);
+
+	float MassRatio = NewMass / OldMass;
+
+
+	MassProps.Mass *= MassRatio;
+	MassProps.InertiaTensor *= MassRatio;
+
+	MassProps.CenterOfMass += MassModifierTransform.TransformVector(OwningBodyInstance->COMNudge);
+
+	// Scale the inertia tensor by the owning body instance's InertiaTensorScale
+	FVector diagonal(MassProps.InertiaTensor.M[0][0] , MassProps.InertiaTensor.M[1][1], MassProps.InertiaTensor.M[2][2]);
+
+	FVector xyz2 = FVector(FVector::DotProduct(diagonal, FVector(0.5f))) - diagonal; // original x^2, y^2, z^2
+	FVector scaledxyz2 = xyz2 * OwningBodyInstance->InertiaTensorScale * OwningBodyInstance->InertiaTensorScale;
+
+	float	xx = scaledxyz2.Y + scaledxyz2.Z,
+			yy = scaledxyz2.Z + scaledxyz2.X,
+			zz = scaledxyz2.X + scaledxyz2.Y;
+
+	float	xy = MassProps.InertiaTensor.M[0][1] * OwningBodyInstance->InertiaTensorScale.X * OwningBodyInstance->InertiaTensorScale.Y,
+			xz = MassProps.InertiaTensor.M[0][2] * OwningBodyInstance->InertiaTensorScale.X * OwningBodyInstance->InertiaTensorScale.Z,
+			yz = MassProps.InertiaTensor.M[1][2] * OwningBodyInstance->InertiaTensorScale.Y * OwningBodyInstance->InertiaTensorScale.Z;
+
+	MassProps.InertiaTensor = Chaos::PMatrix<float, 3, 3>(	FVector(xx, xy, xz),
+															FVector(xy, yy, yz),
+															FVector(xz, yz, zz));
+
+	return MassProps;
+}
+#elif WITH_PHYSX
 /** Computes and adds the mass properties (inertia, com, etc...) based on the mass settings of the body instance. */
 PxMassProperties ComputeMassProperties(const FBodyInstance* OwningBodyInstance, TArray<FPhysicsShapeHandle> Shapes, const FTransform& MassModifierTransform)
 {
@@ -2975,7 +3037,6 @@ void FBodyInstance::UpdateMassProperties()
 {
 	UPhysicalMaterial* PhysMat = GetSimplePhysicalMaterial();
 
-#if WITH_PHYSX
 	if (FPhysicsInterface::IsValid(ActorHandle) && FPhysicsInterface::IsRigidBody(ActorHandle))
 	{
 		FPhysicsCommand::ExecuteWrite(ActorHandle, [&](FPhysicsActorHandle& Actor)
@@ -3003,7 +3064,11 @@ void FBodyInstance::UpdateMassProperties()
 					}
 				}
 
+#if WITH_CHAOS
+				Chaos::TMassProperties<float, 3> TotalMassProperties;
+#elif WITH_PHYSX
 				PxMassProperties TotalMassProperties;
+#endif
 				if (ShapeToBodiesMap.IsValid() && ShapeToBodiesMap->Num() > 0)
 				{
 					struct FWeldedBatch
@@ -3042,6 +3107,21 @@ void FBodyInstance::UpdateMassProperties()
 						}
 					}
 
+#if WITH_CHAOS
+					TArray<Chaos::TMassProperties<float, 3>> SubMassProperties;
+					for (auto BodyShapesItr : BodyToShapes)
+					{
+						const FBodyInstance* OwningBI = BodyShapesItr.Key;
+						const FWeldedBatch& WeldedBatch = BodyShapesItr.Value;
+						FTransform MassModifierTransform = WeldedBatch.RelTM;
+						MassModifierTransform.SetScale3D(MassModifierTransform.GetScale3D() * Scale3D);	//Ensure that any scaling that is done on the component is passed into the mass frame modifiers
+
+						Chaos::TMassProperties<float, 3> BodyMassProperties = ComputeMassProperties(OwningBI, WeldedBatch.Shapes, MassModifierTransform);
+						SubMassProperties.Add(BodyMassProperties);
+					}
+
+					TotalMassProperties = Chaos::Combine(SubMassProperties);
+#elif WITH_PHYSX
 					TArray<PxMassProperties> SubMassProperties;
 					TArray<PxTransform> MassTMs;
 					for (auto BodyShapesItr : BodyToShapes)
@@ -3057,6 +3137,7 @@ void FBodyInstance::UpdateMassProperties()
 					}
 
 					TotalMassProperties = PxMassProperties::sum(SubMassProperties.GetData(), MassTMs.GetData(), SubMassProperties.Num());
+#endif
 				}
 				else
 				{
@@ -3070,25 +3151,29 @@ void FBodyInstance::UpdateMassProperties()
 				}
 
 				// #PHYS2 Refactor out PxMassProperties (Our own impl?)
-				PxQuat MassOrientation;
 #if WITH_CHAOS
-				const PxMat33& IT = TotalMassProperties.inertiaTensor;
-				Chaos::PMatrix<float, 3, 3> InertiaTensor(IT[0][0], IT[0][1], IT[0][2], IT[1][1], IT[1][2], IT[2][2]);
-				const Chaos::TRotation<float, 3> Rotation = Chaos::TransformToLocalSpace(InertiaTensor);
-				const FVector MassSpaceInertiaTensor(InertiaTensor.M[0][0], InertiaTensor.M[1][1], InertiaTensor.M[2][2]);
-				MassOrientation = U2PQuat(Rotation);
+				const Chaos::TRotation<float, 3> Rotation = Chaos::TransformToLocalSpace(TotalMassProperties.InertiaTensor);
+				const FVector MassSpaceInertiaTensor(TotalMassProperties.InertiaTensor.M[0][0], TotalMassProperties.InertiaTensor.M[1][1], TotalMassProperties.InertiaTensor.M[2][2]);
+
+				FPhysicsInterface::SetMassSpaceInertiaTensor_AssumesLocked(Actor, MassSpaceInertiaTensor);
+
+				FPhysicsInterface::SetMass_AssumesLocked(Actor, TotalMassProperties.Mass);
+
+				FTransform Com(Rotation, TotalMassProperties.CenterOfMass);
+				FPhysicsInterface::SetComLocalPose_AssumesLocked(Actor, Com);
 #else
+				PxQuat MassOrientation;
 				const FVector MassSpaceInertiaTensor = P2UVector(PxMassProperties::getMassSpaceInertia(TotalMassProperties.inertiaTensor, MassOrientation));
-#endif
-				FPhysicsInterface::SetMass_AssumesLocked(Actor, TotalMassProperties.mass);
 				FPhysicsInterface::SetMassSpaceInertiaTensor_AssumesLocked(Actor, MassSpaceInertiaTensor);
 
+				FPhysicsInterface::SetMass_AssumesLocked(Actor, TotalMassProperties.mass);
+
 				FTransform Com(P2UQuat(MassOrientation), P2UVector(TotalMassProperties.centerOfMass));
 				FPhysicsInterface::SetComLocalPose_AssumesLocked(Actor, Com);
+#endif
 			}
 		});
 	}
-#endif
 
 	//Let anyone who cares about mass properties know they've been updated
 	if (BodyInstanceDelegates.IsValid())
diff --git a/Engine/Source/Runtime/Engine/Private/PhysicsEngine/Experimental/PhysInterface_Chaos.cpp b/Engine/Source/Runtime/Engine/Private/PhysicsEngine/Experimental/PhysInterface_Chaos.cpp
--- a/Engine/Source/Runtime/Engine/Private/PhysicsEngine/Experimental/PhysInterface_Chaos.cpp
+++ b/Engine/Source/Runtime/Engine/Private/PhysicsEngine/Experimental/PhysInterface_Chaos.cpp
@@ -39,13 +39,13 @@
 #include "Components/SkeletalMeshComponent.h"
 #include "PBDRigidsSolver.h"
 #include "PhysicalMaterials/PhysicalMaterialMask.h"
+#include "PhysicalMaterials/PhysicalMaterial.h"
 
 #if WITH_PHYSX
 #include "geometry/PxConvexMesh.h"
 #include "geometry/PxTriangleMesh.h"
 #include "foundation/PxVec3.h"
 #include "extensions/PxMassProperties.h"
-#include "PhysicalMaterials/PhysicalMaterial.h"
 #include "Containers/ArrayView.h"
 #endif
 
@@ -1739,18 +1739,16 @@ void FinishSceneStat()
 {
 }
 
-#if WITH_PHYSX
-
 bool CalculateMassPropertiesOfImplicitType(
 	Chaos::TMassProperties<float, 3>& OutMassProperties,
-	const Chaos::TRigidTransform<float, 3> & WorldTransform,
-	const Chaos::FImplicitObject* ImplicitObject, 
+	const Chaos::TRigidTransform<float, 3>& WorldTransform,
+	const Chaos::FImplicitObject* ImplicitObject,
 	float InDensityKGPerCM)
 {
 	// WIP
 	// @todo : Support center of mass offsets.
 	// @todo : Support Mass space alignment. 
-	
+
 	using namespace Chaos;
 
 	if (ImplicitObject)
@@ -1775,8 +1773,8 @@ bool CalculateMassPropertiesOfImplicitType(
 				OutMassProperties.RotationOfMass = Chaos::TRotation<float, 3>::FromIdentity();
 				return false;
 			}
-		else*/ 
-		
+		else*/
+
 		//todo: Still need to handle scaled
 		Chaos::Utilities::CastHelper(*ImplicitObject, FTransform::Identity, [&OutMassProperties, InDensityKGPerCM](const auto& Object, const auto& LocalTM)
 			{
@@ -1791,6 +1789,53 @@ bool CalculateMassPropertiesOfImplicitType(
 	return false;
 }
 
+void FPhysInterface_Chaos::CalculateTMassPropertiesFromShapeCollection(Chaos::TMassProperties<float, 3>& OutProperties, const TArray<FPhysicsShapeHandle>& InShapes, float InDensityKGPerCM)
+{
+	float TotalMass = 0.f;
+	Chaos::FVec3 TotalCenterOfMass(0.f);
+	TArray< Chaos::TMassProperties<float, 3> > MassPropertiesList;
+	for (const FPhysicsShapeHandle& ShapeHandle : InShapes)
+	{
+		if (const Chaos::TPerShapeData<float, 3>* Shape = ShapeHandle.Shape)
+		{
+			if (const Chaos::FImplicitObject* ImplicitObject = Shape->Geometry.Get())
+			{
+				FTransform WorldTransform(ShapeHandle.ActorRef->R(), ShapeHandle.ActorRef->X());
+				Chaos::TMassProperties<float, 3> MassProperties;
+				if (CalculateMassPropertiesOfImplicitType(MassProperties, WorldTransform, ImplicitObject, InDensityKGPerCM))
+				{
+					MassPropertiesList.Add(MassProperties);
+					TotalMass += MassProperties.Mass;
+					TotalCenterOfMass += MassProperties.CenterOfMass * MassProperties.Mass;
+				}
+			}
+		}
+	}
+
+	if (TotalMass > 0.f)
+	{
+		TotalCenterOfMass /= TotalMass;
+	}
+
+	Chaos::PMatrix<float, 3, 3> Tensor;
+	if (MassPropertiesList.Num())
+	{
+		Tensor = Chaos::CombineWorldSpace<float, 3>(MassPropertiesList, InDensityKGPerCM).InertiaTensor;
+	}
+	else
+	{
+		// @todo : Add support for all types, but for now just hard code a unit sphere tensor {r:50cm} if the type was not processed
+		Tensor = Chaos::PMatrix<float, 3, 3>(5.24e5, 5.24e5, 5.24e5);
+		TotalMass = 523.f;
+	}
+
+	OutProperties.InertiaTensor = Tensor;
+	OutProperties.Mass = TotalMass;
+	OutProperties.CenterOfMass = TotalCenterOfMass;
+}
+
+#if WITH_PHYSX
+
 void FPhysInterface_Chaos::CalculateMassPropertiesFromShapeCollection(physx::PxMassProperties& OutProperties, const TArray<FPhysicsShapeHandle>& InShapes, float InDensityKGPerCM)
 {
 	float TotalMass = 0.f;
diff --git a/Engine/Source/Runtime/Engine/Private/PhysicsEngine/PhysicsInterfaceUtils.cpp b/Engine/Source/Runtime/Engine/Private/PhysicsEngine/PhysicsInterfaceUtils.cpp
--- a/Engine/Source/Runtime/Engine/Private/PhysicsEngine/PhysicsInterfaceUtils.cpp
+++ b/Engine/Source/Runtime/Engine/Private/PhysicsEngine/PhysicsInterfaceUtils.cpp
@@ -6,8 +6,6 @@
 #include "Physics/PhysicsFiltering.h"
 #include "Physics/PhysicsInterfaceTypes.h"
 
-#if WITH_PHYSX
-
 FCollisionFilterData CreateObjectQueryFilterData(const bool bTraceComplex, const int32 MultiTrace/*=1 if multi. 0 otherwise*/, const struct FCollisionObjectQueryParams & ObjectParam)
 {
 	/**
@@ -93,12 +91,10 @@ FCollisionFilterData CreateTraceQueryFilterData(const uint8 MyChannel, const boo
 #define TRACE_MULTI		1
 #define TRACE_SINGLE	0
 
-#endif
-
 /** Utility for creating a PhysX PxFilterData for performing a query (trace) against the scene */
 FCollisionFilterData CreateQueryFilterData(const uint8 MyChannel, const bool bTraceComplex, const FCollisionResponseContainer& InCollisionResponseContainer, const struct FCollisionQueryParams& QueryParam, const struct FCollisionObjectQueryParams & ObjectParam, const bool bMultitrace)
 {
-#if WITH_PHYSX
+	//#TODO implement chaos
 	if(ObjectParam.IsValid())
 	{
 		return CreateObjectQueryFilterData(bTraceComplex, (bMultitrace ? TRACE_MULTI : TRACE_SINGLE), ObjectParam);
@@ -107,9 +103,4 @@ FCollisionFilterData CreateQueryFilterData(const uint8 MyChannel, const bool bTr
 	{
 		return CreateTraceQueryFilterData(MyChannel, bTraceComplex, InCollisionResponseContainer, QueryParam);
 	}
-#else 
-	//#TODO implement chaos
-	ensure(false);
-	return FCollisionFilterData();
-#endif
 }
diff --git a/Engine/Source/Runtime/Engine/Private/SkeletalMeshComponentPhysics.cpp b/Engine/Source/Runtime/Engine/Private/SkeletalMeshComponentPhysics.cpp
--- a/Engine/Source/Runtime/Engine/Private/SkeletalMeshComponentPhysics.cpp
+++ b/Engine/Source/Runtime/Engine/Private/SkeletalMeshComponentPhysics.cpp
@@ -583,7 +583,6 @@ void USkeletalMeshComponent::InitArticulated(FPhysScene* PhysScene)
 	}
 
 	// Set up the map from skelmeshcomp ID to collision disable table
-#if WITH_PHYSX
 	uint32 SkelMeshCompID = GetUniqueID();
 	PhysScene->DeferredAddCollisionDisableTable(SkelMeshCompID, &PhysicsAsset->CollisionDisableTable);
 
@@ -607,7 +606,6 @@ void USkeletalMeshComponent::InitArticulated(FPhysScene* PhysScene)
 	}
 
 	InstantiatePhysicsAsset(*PhysicsAsset, Scale3D, Bodies, Constraints, PhysScene, this, RootBodyIndex, Aggregate);
-#endif // WITH_PHYSX
 
 	// now update root body index because body has BodySetup now
 	SetRootBodyIndex(RootBodyIndex);
@@ -739,7 +737,6 @@ void USkeletalMeshComponent::InstantiatePhysicsAsset_Internal(const UPhysicsAsse
 				}
 			}
 
-#if WITH_PHYSX
 			// Create physics body instance.
 			FTransform BoneTransform = BoneTransformGetter(BoneIndex);
 
@@ -756,18 +753,14 @@ void USkeletalMeshComponent::InstantiatePhysicsAsset_Internal(const UPhysicsAsse
 			BodyInst->InitBody(PhysicsAssetBodySetup, BoneTransform, OwningComponent, PhysScene, SpawnParams);
 
 			NameToBodyMap.Add(PhysicsAssetBodySetup->BoneName, BodyInst);
-#endif //WITH_PHYSX
 		}
 	}
 
-#if WITH_PHYSX
 	if(PhysScene && Aggregate.IsValid())
 	{
 		PhysScene->AddAggregateToScene(Aggregate);
 	}
-
-#endif //WITH_PHYSX
-
+	
 	// Create all the OutConstraints.
 	check(OutConstraints.Num() == 0);
 	int32 NumOutConstraints = PhysAsset.ConstraintSetup.Num();
@@ -838,7 +831,6 @@ void USkeletalMeshComponent::TermArticulated()
 {
 	ResetRootBodyIndex();
 
-#if WITH_PHYSX
 	uint32 SkelMeshCompID = GetUniqueID();
 	UWorld* MyWorld = GetWorld();
 	FPhysScene* PhysScene = (MyWorld ? MyWorld->GetPhysicsScene() : nullptr);
@@ -846,7 +838,6 @@ void USkeletalMeshComponent::TermArticulated()
 	{
 		PhysScene->DeferredRemoveCollisionDisableTable(SkelMeshCompID);
 	}
-#endif	//#if WITH_PHYSX
 
 	FPhysicsCommand::ExecuteWrite(this, [&]()
 	{
@@ -873,14 +864,13 @@ void USkeletalMeshComponent::TermArticulated()
 	
 	Bodies.Empty();
 
-#if WITH_PHYSX
 	// releasing Aggregate, it shouldn't contain any Bodies now, because they are released above
 		if(Aggregate.IsValid())
 	{
 			check(FPhysicsInterface::GetNumActorsInAggregate(Aggregate) == 0);
 			FPhysicsInterface::ReleaseAggregate(Aggregate);
 	}
-#endif //WITH_PHYSX
+
 	});
 }
 
diff --git a/Engine/Source/Runtime/Engine/Public/Physics/Experimental/ChaosInterfaceWrapper.h b/Engine/Source/Runtime/Engine/Public/Physics/Experimental/ChaosInterfaceWrapper.h
--- a/Engine/Source/Runtime/Engine/Public/Physics/Experimental/ChaosInterfaceWrapper.h
+++ b/Engine/Source/Runtime/Engine/Public/Physics/Experimental/ChaosInterfaceWrapper.h
@@ -35,11 +35,7 @@ struct FScopedSceneReadLock
 
 inline FQueryFilterData MakeQueryFilterData(const FCollisionFilterData& FilterData, EQueryFlags QueryFlags, const FCollisionQueryParams& Params)
 {
-#if WITH_PHYSX
 	return PxQueryFilterData(U2PFilterData(FilterData), U2PQueryFlags(QueryFlags));
-#else
-	return FQueryFilterData();
-#endif
 }
 
 FBodyInstance* GetUserData(const Chaos::TGeometryParticle<float, 3>& Actor);
diff --git a/Engine/Source/Runtime/Engine/Public/Physics/Experimental/PhysInterface_Chaos.h b/Engine/Source/Runtime/Engine/Public/Physics/Experimental/PhysInterface_Chaos.h
--- a/Engine/Source/Runtime/Engine/Public/Physics/Experimental/PhysInterface_Chaos.h
+++ b/Engine/Source/Runtime/Engine/Public/Physics/Experimental/PhysInterface_Chaos.h
@@ -54,6 +54,9 @@ namespace Chaos
 
 	template <typename T>
 	class TCapsule;
+	
+	template<class T, int d>
+	struct TMassProperties;
 
 	template <typename T, int>
 	class TAABB;
@@ -371,6 +374,7 @@ public:
 
 	static bool ExecPhysCommands(const TCHAR* Cmd, FOutputDevice* Ar, UWorld* InWorld);
 
+	static void CalculateTMassPropertiesFromShapeCollection(Chaos::TMassProperties<float, 3>& OutProperties, const TArray<FPhysicsShapeHandle>& InShapes, float InDensityKGPerCM);
 #if WITH_PHYSX
     static void CalculateMassPropertiesFromShapeCollection(physx::PxMassProperties& OutProperties, const TArray<FPhysicsShapeHandle>& InShapes, float InDensityKGPerCM);
 #endif
diff --git a/Engine/Source/Runtime/Engine/Public/Physics/PhysicsInterfaceUtils.h b/Engine/Source/Runtime/Engine/Public/Physics/PhysicsInterfaceUtils.h
--- a/Engine/Source/Runtime/Engine/Public/Physics/PhysicsInterfaceUtils.h
+++ b/Engine/Source/Runtime/Engine/Public/Physics/PhysicsInterfaceUtils.h
@@ -30,6 +30,8 @@ struct FConstraintBrokenDelegateData
 	int32 ConstraintIndex;
 };
 
+#endif // WITH_PHYSX
+
 class FPhysicsReplication;
 
 /** Interface for the creation of customized physics replication.*/
@@ -39,4 +41,3 @@ public:
 	virtual FPhysicsReplication* Create(FPhysScene* OwningPhysScene) = 0;
 	virtual void Destroy(FPhysicsReplication* PhysicsReplication) = 0;
 };
-#endif // WITH_PHYX
diff --git a/Engine/Source/Runtime/Landscape/Private/LandscapeCollision.cpp b/Engine/Source/Runtime/Landscape/Private/LandscapeCollision.cpp
--- a/Engine/Source/Runtime/Landscape/Private/LandscapeCollision.cpp
+++ b/Engine/Source/Runtime/Landscape/Private/LandscapeCollision.cpp
@@ -252,7 +252,6 @@ void ULandscapeHeightfieldCollisionComponent::OnCreatePhysicsState()
 
 	if (!BodyInstance.IsValidBodyInstance())
 	{
-#if WITH_PHYSX
 		CreateCollisionObject();
 
 		if (IsValidRef(HeightfieldRef))
@@ -280,9 +279,7 @@ void ULandscapeHeightfieldCollisionComponent::OnCreatePhysicsState()
 			LandscapeComponentMatrix.SetAxis(0, TerrainX);
 			LandscapeComponentMatrix.SetAxis(2, TerrainY);
 			LandscapeComponentMatrix.SetAxis(1, TerrainZ);
-
-			PxTransform PhysXLandscapeComponentTransform = U2PTransform(FTransform(LandscapeComponentMatrix));
-
+			
 			const bool bCreateSimpleCollision = SimpleCollisionSizeQuads > 0;
 			const float SimpleCollisionScale = bCreateSimpleCollision ? CollisionScale * CollisionSizeQuads / SimpleCollisionSizeQuads : 0;
 
@@ -290,6 +287,7 @@ void ULandscapeHeightfieldCollisionComponent::OnCreatePhysicsState()
 			FVector FinalScale(LandscapeScale.X * CollisionScale, LandscapeScale.Y * CollisionScale, LandscapeScale.Z * LANDSCAPE_ZSCALE);
 
 #if PHYSICS_INTERFACE_PHYSX
+			PxTransform PhysXLandscapeComponentTransform = U2PTransform(FTransform(LandscapeComponentMatrix));
 			PxHeightFieldGeometry LandscapeComponentGeom(HeightfieldRef->RBHeightfield, PxMeshGeometryFlag::eDOUBLE_SIDED, LandscapeScale.Z * LANDSCAPE_ZSCALE, LandscapeScale.Y * CollisionScale, LandscapeScale.X * CollisionScale);
 
 			if (LandscapeComponentGeom.isValid())
@@ -523,7 +521,6 @@ void ULandscapeHeightfieldCollisionComponent::OnCreatePhysicsState()
 			}
 #endif // WITH_CHAOS
 		}
-#endif// WITH_PHYSX
 	}
 }
 
diff --git a/Engine/Source/Runtime/PhysicsCore/Private/PhysicsInterfaceUtilsCore.cpp b/Engine/Source/Runtime/PhysicsCore/Private/PhysicsInterfaceUtilsCore.cpp
--- a/Engine/Source/Runtime/PhysicsCore/Private/PhysicsInterfaceUtilsCore.cpp
+++ b/Engine/Source/Runtime/PhysicsCore/Private/PhysicsInterfaceUtilsCore.cpp
@@ -6,6 +6,21 @@
 
 #include "Chaos/ParticleHandle.h"
 
+FCollisionFilterData P2UFilterData(const PxFilterData& PFilterData)
+{
+	FCollisionFilterData FilterData;
+	FilterData.Word0 = PFilterData.word0;
+	FilterData.Word1 = PFilterData.word1;
+	FilterData.Word2 = PFilterData.word2;
+	FilterData.Word3 = PFilterData.word3;
+	return FilterData;
+}
+
+PxFilterData U2PFilterData(const FCollisionFilterData& FilterData)
+{
+	return PxFilterData(FilterData.Word0, FilterData.Word1, FilterData.Word2, FilterData.Word3);
+}
+
 #if WITH_PHYSX
 
 PxShapeFlags BuildPhysXShapeFlags(FBodyCollisionFlags BodyCollisionFlags, bool bPhysicsStatic, bool bIsTriangleMesh)
@@ -29,21 +44,6 @@ PxShapeFlags BuildPhysXShapeFlags(FBodyCollisionFlags BodyCollisionFlags, bool b
 	return ShapeFlags;
 }
 
-PxFilterData U2PFilterData(const FCollisionFilterData& FilterData)
-{
-	return PxFilterData(FilterData.Word0, FilterData.Word1, FilterData.Word2, FilterData.Word3);
-}
-
-FCollisionFilterData P2UFilterData(const PxFilterData& PFilterData)
-{
-	FCollisionFilterData FilterData;
-	FilterData.Word0 = PFilterData.word0;
-	FilterData.Word1 = PFilterData.word1;
-	FilterData.Word2 = PFilterData.word2;
-	FilterData.Word3 = PFilterData.word3;
-	return FilterData;
-}
-
 PxGeometryType::Enum U2PCollisionShapeType(ECollisionShapeType InUType)
 {
 	switch (InUType)
diff --git a/Engine/Source/Runtime/PhysicsCore/Private/SQCapture.cpp b/Engine/Source/Runtime/PhysicsCore/Private/SQCapture.cpp
--- a/Engine/Source/Runtime/PhysicsCore/Private/SQCapture.cpp
+++ b/Engine/Source/Runtime/PhysicsCore/Private/SQCapture.cpp
@@ -174,6 +174,8 @@
 	}
 }
 
+#endif
+
 void SerializeQueryFilterData(FArchive& Ar, FQueryFilterData& QueryFilterData)
 {
 	Ar << QueryFilterData.data.word0;
@@ -186,8 +188,6 @@
 	Ar << QueryFilterData.clientId;
 }
 
-#endif
-
 void FSQCapture::SerializeChaosActorToShapeHitsArray(Chaos::FChaosArchive& Ar)
 {
 	int32 NumActors = ChaosActorToShapeHitsArray.Num();
@@ -331,7 +331,9 @@
 
 	if (Ar.IsLoading())
 	{
+#if WITH_PHYSX
 		CreateChaosDataFromPhysX();
+#endif // WITH_PHYSX
 
 		if (bDiskDataIsChaos)
 		{
@@ -419,6 +421,28 @@
 	}
 }
 
+void FSQCapture::CaptureChaosFilterResults(const Chaos::FPBDRigidsEvolutionGBF& TransientEvolution, const FCollisionFilterData& FilterData, ICollisionQueryFilterCallbackBase& Callback)
+{
+	using namespace Chaos;
+	const TPBDRigidsSOAs<float, 3>& Particles = TransientEvolution.GetParticles();
+	const int32 NumTransientActors = Particles.GetParticleHandles().Size();
+
+	for (int32 Idx = 0; Idx < NumTransientActors; ++Idx)
+	{
+		TGeometryParticle<float, 3>* TransientActor = Particles.GetParticleHandles().Handle(Idx)->GTGeometryParticle();
+		const TShapesArray<float, 3>& TransientShapes = TransientActor->ShapesArray();
+		const int32 NumTransientShapes = TransientShapes.Num();
+		TArray<TPair<TPerShapeData<float, 3>*, ECollisionQueryHitType>> ShapeHitsArray; ShapeHitsArray.Reserve(NumTransientShapes);
+		for (const auto& TransientShape : TransientShapes)
+		{
+			const ECollisionQueryHitType Result = Callback.PreFilter(FilterData, *TransientShape, *TransientActor);
+			ShapeHitsArray.Emplace(TransientShape.Get(), Result);
+		}
+
+		ChaosActorToShapeHitsArray.Add(TransientActor, ShapeHitsArray);
+	}
+}
+
 #if WITH_PHYSX
 void FSQCapture::StartCapturePhysXSweep(const PxScene& Scene, const PxGeometry& InQueryGeom, const FTransform& InStartTM, const FVector& InDir, float InDeltaMag, FHitFlags InOutputFlags, const FQueryFilterData& QueryFilter, const FCollisionFilterData& FilterData, ICollisionQueryFilterCallbackBase& Callback)
 {
@@ -506,30 +530,7 @@
 	}
 }
 
-#if WITH_PHYSX
-void FSQCapture::CaptureChaosFilterResults(const Chaos::FPBDRigidsEvolutionGBF& TransientEvolution, const FCollisionFilterData& FilterData, ICollisionQueryFilterCallbackBase& Callback)
-{
-	using namespace Chaos;
-	const TPBDRigidsSOAs<float, 3>& Particles = TransientEvolution.GetParticles();
-	const int32 NumTransientActors = Particles.GetParticleHandles().Size();
-	
-	for (int32 Idx = 0; Idx < NumTransientActors; ++Idx)
-	{
-		TGeometryParticle<float, 3>* TransientActor = Particles.GetParticleHandles().Handle(Idx)->GTGeometryParticle();
-		const TShapesArray<float,3>& TransientShapes = TransientActor->ShapesArray();
-		const int32 NumTransientShapes = TransientShapes.Num();
-		TArray<TPair<TPerShapeData<float,3>*, ECollisionQueryHitType>> ShapeHitsArray; ShapeHitsArray.Reserve(NumTransientShapes);
-		for (const auto& TransientShape : TransientShapes)
-		{
-			const ECollisionQueryHitType Result = Callback.PreFilter(FilterData, *TransientShape, *TransientActor);
-			ShapeHitsArray.Emplace(TransientShape.Get(), Result);
-		}
+#endif // WITH_PHYSX
 
-		ChaosActorToShapeHitsArray.Add(TransientActor, ShapeHitsArray);
-	}
-}
-#endif
-#endif
-
 template <typename THit>
 void EndCaptureHelper(PhysXInterface::FDynamicHitBuffer<THit>& Dest, const PxHitCallback<THit>& Results)
@@ -582,6 +583,30 @@
 }
 #endif
 
+template <typename TShape, typename TActor>
+ECollisionQueryHitType GetFilterResultHelper(const TShape* Shape, const TActor* Actor, const TMap<TActor*, TArray<TPair<TShape*, ECollisionQueryHitType>>>& ActorToShapeHitsArray)
+{
+	if (const TArray<TPair<TShape*, ECollisionQueryHitType>>* ActorToPairs = ActorToShapeHitsArray.Find(Actor))
+	{
+		for (const TPair<TShape*, ECollisionQueryHitType>& Pair : *ActorToPairs)
+		{
+			if (Pair.Key == Shape)
+			{
+				return Pair.Value;
+			}
+		}
+	}
+
+	//todo: figure out why this hits - suspect it's related to threading and how we capture an evolution on GT
+	ensure(false);	//should not get here, means we didn't properly capture all filter results
+	return ECollisionQueryHitType::None;
+}
+
+ECollisionQueryHitType FSQCapture::GetFilterResult(const Chaos::TPerShapeData<float, 3>* Shape, const Chaos::TGeometryParticle<float, 3>* Actor) const
+{
+	return GetFilterResultHelper(Shape, Actor, ChaosActorToShapeHitsArray);
+}
+
 #if WITH_PHYSX
 
 template <typename THit>
@@ -606,35 +631,11 @@
 	}
 }
 
-template <typename TShape, typename TActor>
-ECollisionQueryHitType GetFilterResultHelper(const TShape* Shape, const TActor* Actor, const TMap<TActor*, TArray<TPair<TShape*, ECollisionQueryHitType>>>& ActorToShapeHitsArray)
-{
-	if (const TArray<TPair<TShape*, ECollisionQueryHitType>>* ActorToPairs = ActorToShapeHitsArray.Find(Actor))
-	{
-		for (const TPair<TShape*, ECollisionQueryHitType>& Pair : *ActorToPairs)
-		{
-			if (Pair.Key == Shape)
-			{
-				return Pair.Value;
-			}
-		}
-	}
-
-	//todo: figure out why this hits - suspect it's related to threading and how we capture an evolution on GT
-	ensure(false);	//should not get here, means we didn't properly capture all filter results
-	return ECollisionQueryHitType::None;
-}
-
 ECollisionQueryHitType FSQCapture::GetFilterResult(const PxShape* Shape, const PxActor* Actor) const
 {
 	return GetFilterResultHelper(Shape, Actor, PxActorToShapeHitsArray);
 }
 
-ECollisionQueryHitType FSQCapture::GetFilterResult(const Chaos::TPerShapeData<float,3>* Shape, const Chaos::TGeometryParticle<float,3>* Actor) const
-{
-	return GetFilterResultHelper(Shape, Actor, ChaosActorToShapeHitsArray);
-}
-
 void FSQCapture::CreatePhysXData()
 {
 	if (bDiskDataIsChaos || bPhysXDataReady)
diff --git a/Engine/Source/Runtime/PhysicsCore/Public/ChaosInterfaceWrapperCore.h b/Engine/Source/Runtime/PhysicsCore/Public/ChaosInterfaceWrapperCore.h
--- a/Engine/Source/Runtime/PhysicsCore/Public/ChaosInterfaceWrapperCore.h
+++ b/Engine/Source/Runtime/PhysicsCore/Public/ChaosInterfaceWrapperCore.h
@@ -28,11 +28,7 @@ struct FDummyPhysActor {};
 template<typename DummyT>
 struct FDummyCallback {};
 
-#if WITH_PHYSX
 using FQueryFilterData = PxQueryFilterData;
-#else
-using FQueryFilterData = FDummyPhysType;
-#endif
 
 /** We use this struct so that if no conversion is needed in another API, we can avoid the copy (if we think that's critical) */
 struct FPhysicsRaycastInputAdapater
diff --git a/Engine/Source/Runtime/PhysicsCore/Public/PhysicsInterfaceTypesCore.h b/Engine/Source/Runtime/PhysicsCore/Public/PhysicsInterfaceTypesCore.h
--- a/Engine/Source/Runtime/PhysicsCore/Public/PhysicsInterfaceTypesCore.h
+++ b/Engine/Source/Runtime/PhysicsCore/Public/PhysicsInterfaceTypesCore.h
@@ -141,4 +141,451 @@ static bool CalcMeshNegScaleCompensation(const FVector& InScale3D, FTransform& O
 
 	// Use inverted mesh if determinant is negative
 	return (InScale3D.X * InScale3D.Y * InScale3D.Z) < 0.f;
-}
\ No newline at end of file
+}
+
+
+#if !WITH_PHYSX
+
+const uint32 AggregateMaxSize = 128;
+
+/** enum for empty constructor tag*/
+enum PxEMPTY
+{
+	PxEmpty
+};
+
+#define PX_INLINE inline
+
+typedef int64_t PxI64;
+typedef uint64_t PxU64;
+typedef int32_t PxI32;
+typedef uint32_t PxU32;
+typedef int16_t PxI16;
+typedef uint16_t PxU16;
+typedef int8_t PxI8;
+typedef uint8_t PxU8;
+typedef float PxF32;
+typedef double PxF64;
+typedef float PxReal;
+
+struct PxFilterData
+{
+	//= ATTENTION! =====================================================================================
+	// Changing the data layout of this class breaks the binary serialization format.  See comments for 
+	// PX_BINARY_SERIAL_VERSION.  If a modification is required, please adjust the getBinaryMetaData 
+	// function.  If the modification is made on a custom branch, please change PX_BINARY_SERIAL_VERSION
+	// accordingly.
+	//==================================================================================================
+
+	PX_INLINE PxFilterData(const PxEMPTY)
+	{
+	}
+
+	/**
+	\brief Default constructor.
+	*/
+	PX_INLINE PxFilterData()
+	{
+		word0 = word1 = word2 = word3 = 0;
+	}
+
+	/**
+	\brief Constructor to set filter data initially.
+	*/
+	PX_INLINE PxFilterData(PxU32 w0, PxU32 w1, PxU32 w2, PxU32 w3) : word0(w0), word1(w1), word2(w2), word3(w3) {}
+
+	/**
+	\brief (re)sets the structure to the default.
+	*/
+	PX_INLINE void setToDefault()
+	{
+		*this = PxFilterData();
+	}
+
+	/**
+	\brief Comparison operator to allow use in Array.
+	*/
+	PX_INLINE bool operator == (const PxFilterData& a) const
+	{
+		return a.word0 == word0 && a.word1 == word1 && a.word2 == word2 && a.word3 == word3;
+	}
+
+	/**
+	\brief Comparison operator to allow use in Array.
+	*/
+	PX_INLINE bool operator != (const PxFilterData& a) const
+	{
+		return !(a == *this);
+	}
+
+	PxU32 word0;
+	PxU32 word1;
+	PxU32 word2;
+	PxU32 word3;
+};
+
+struct PxQueryFlag
+{
+	enum Enum
+	{
+		eSTATIC = (1 << 0),	//!< Traverse static shapes
+
+		eDYNAMIC = (1 << 1),	//!< Traverse dynamic shapes
+
+		ePREFILTER = (1 << 2),	//!< Run the pre-intersection-test filter (see #PxQueryFilterCallback::preFilter())
+
+		ePOSTFILTER = (1 << 3),	//!< Run the post-intersection-test filter (see #PxQueryFilterCallback::postFilter())
+
+		eANY_HIT = (1 << 4),	//!< Abort traversal as soon as any hit is found and return it via callback.block.
+								//!< Helps query performance. Both eTOUCH and eBLOCK hitTypes are considered hits with this flag.
+
+		eNO_BLOCK = (1 << 5),	//!< All hits are reported as touching. Overrides eBLOCK returned from user filters with eTOUCH.
+								//!< This is also an optimization hint that may improve query performance.
+
+		eRESERVED = (1 << 15)	//!< Reserved for internal use
+	};
+};
+
+typedef PxU8 PxClientID;
+static const PxClientID PX_DEFAULT_CLIENT = 0;
+
+#define PX_CUDA_CALLABLE
+
+template <typename enumtype, typename storagetype = uint32_t>
+class PxFlags
+{
+public:
+	typedef storagetype InternalType;
+
+	PX_CUDA_CALLABLE PX_INLINE explicit PxFlags(const PxEMPTY)
+	{
+	}
+	PX_CUDA_CALLABLE PX_INLINE PxFlags(void);
+	PX_CUDA_CALLABLE PX_INLINE PxFlags(enumtype e);
+	PX_CUDA_CALLABLE PX_INLINE PxFlags(const PxFlags<enumtype, storagetype>& f);
+	PX_CUDA_CALLABLE PX_INLINE explicit PxFlags(storagetype b);
+
+	PX_CUDA_CALLABLE PX_INLINE bool isSet(enumtype e) const;
+	PX_CUDA_CALLABLE PX_INLINE PxFlags<enumtype, storagetype>& set(enumtype e);
+	PX_CUDA_CALLABLE PX_INLINE bool operator==(enumtype e) const;
+	PX_CUDA_CALLABLE PX_INLINE bool operator==(const PxFlags<enumtype, storagetype>& f) const;
+	PX_CUDA_CALLABLE PX_INLINE bool operator==(bool b) const;
+	PX_CUDA_CALLABLE PX_INLINE bool operator!=(enumtype e) const;
+	PX_CUDA_CALLABLE PX_INLINE bool operator!=(const PxFlags<enumtype, storagetype>& f) const;
+
+	PX_CUDA_CALLABLE PX_INLINE PxFlags<enumtype, storagetype>& operator=(const PxFlags<enumtype, storagetype>& f);
+	PX_CUDA_CALLABLE PX_INLINE PxFlags<enumtype, storagetype>& operator=(enumtype e);
+
+	PX_CUDA_CALLABLE PX_INLINE PxFlags<enumtype, storagetype>& operator|=(enumtype e);
+	PX_CUDA_CALLABLE PX_INLINE PxFlags<enumtype, storagetype>& operator|=(const PxFlags<enumtype, storagetype>& f);
+	PX_CUDA_CALLABLE PX_INLINE PxFlags<enumtype, storagetype> operator|(enumtype e) const;
+	PX_CUDA_CALLABLE PX_INLINE PxFlags<enumtype, storagetype> operator|(const PxFlags<enumtype, storagetype>& f) const;
+
+	PX_CUDA_CALLABLE PX_INLINE PxFlags<enumtype, storagetype>& operator&=(enumtype e);
+	PX_CUDA_CALLABLE PX_INLINE PxFlags<enumtype, storagetype>& operator&=(const PxFlags<enumtype, storagetype>& f);
+	PX_CUDA_CALLABLE PX_INLINE PxFlags<enumtype, storagetype> operator&(enumtype e) const;
+	PX_CUDA_CALLABLE PX_INLINE PxFlags<enumtype, storagetype> operator&(const PxFlags<enumtype, storagetype>& f) const;
+
+	PX_CUDA_CALLABLE PX_INLINE PxFlags<enumtype, storagetype>& operator^=(enumtype e);
+	PX_CUDA_CALLABLE PX_INLINE PxFlags<enumtype, storagetype>& operator^=(const PxFlags<enumtype, storagetype>& f);
+	PX_CUDA_CALLABLE PX_INLINE PxFlags<enumtype, storagetype> operator^(enumtype e) const;
+	PX_CUDA_CALLABLE PX_INLINE PxFlags<enumtype, storagetype> operator^(const PxFlags<enumtype, storagetype>& f) const;
+
+	PX_CUDA_CALLABLE PX_INLINE PxFlags<enumtype, storagetype> operator~(void) const;
+
+	PX_CUDA_CALLABLE PX_INLINE operator bool(void) const;
+	PX_CUDA_CALLABLE PX_INLINE operator uint8_t(void) const;
+	PX_CUDA_CALLABLE PX_INLINE operator uint16_t(void) const;
+	PX_CUDA_CALLABLE PX_INLINE operator uint32_t(void) const;
+
+	PX_CUDA_CALLABLE PX_INLINE void clear(enumtype e);
+
+public:
+	friend PX_INLINE PxFlags<enumtype, storagetype> operator&(enumtype a, PxFlags<enumtype, storagetype>& b)
+	{
+		PxFlags<enumtype, storagetype> out;
+		out.mBits = a & b.mBits;
+		return out;
+	}
+
+private:
+	storagetype mBits;
+};
+
+typedef PxFlags<PxQueryFlag::Enum, PxU16> PxQueryFlags;
+
+inline PxQueryFlags U2PQueryFlags(FQueryFlags Flags)
+{
+	uint32 Result = 0;
+	if (Flags & EQueryFlags::PreFilter)
+	{
+		Result |= PxQueryFlag::ePREFILTER;
+	}
+
+	if (Flags & EQueryFlags::PostFilter)
+	{
+		Result |= PxQueryFlag::ePOSTFILTER;
+	}
+
+	if (Flags & EQueryFlags::AnyHit)
+	{
+		Result |= PxQueryFlag::eANY_HIT;
+	}
+
+	return (PxQueryFlags)Result;
+}
+
+struct PxQueryFilterData
+{
+	/** \brief default constructor */
+	explicit PX_INLINE PxQueryFilterData() : flags(PxQueryFlag::eDYNAMIC | PxQueryFlag::eSTATIC), clientId(PX_DEFAULT_CLIENT) {}
+
+	/** \brief constructor to set both filter data and filter flags */
+	explicit PX_INLINE PxQueryFilterData(const PxFilterData& fd, PxQueryFlags f) : data(fd), flags(f), clientId(PX_DEFAULT_CLIENT) {}
+
+	/** \brief constructor to set filter flags only */
+	explicit PX_INLINE PxQueryFilterData(PxQueryFlags f) : flags(f), clientId(PX_DEFAULT_CLIENT) {}
+
+	PxFilterData	data;		//!< Filter data associated with the scene query
+	PxQueryFlags	flags;		//!< Filter flags (see #PxQueryFlags)
+	PxClientID		clientId;	//!< ID of the client doing the query (see #PxScene.createClient())
+};
+
+#define PX_FLAGS_OPERATORS(enumtype, storagetype)                                                                      \
+	PX_INLINE PxFlags<enumtype, storagetype> operator|(enumtype a, enumtype b)                                         \
+	{                                                                                                                  \
+		PxFlags<enumtype, storagetype> r(a);                                                                           \
+		r |= b;                                                                                                        \
+		return r;                                                                                                      \
+	}                                                                                                                  \
+	PX_INLINE PxFlags<enumtype, storagetype> operator&(enumtype a, enumtype b)                                         \
+	{                                                                                                                  \
+		PxFlags<enumtype, storagetype> r(a);                                                                           \
+		r &= b;                                                                                                        \
+		return r;                                                                                                      \
+	}                                                                                                                  \
+	PX_INLINE PxFlags<enumtype, storagetype> operator~(enumtype a)                                                     \
+	{                                                                                                                  \
+		return ~PxFlags<enumtype, storagetype>(a);                                                                     \
+	}
+
+#define PX_FLAGS_TYPEDEF(x, y)                                                                                         \
+	typedef PxFlags<x::Enum, y> x##s;                                                                                  \
+	PX_FLAGS_OPERATORS(x::Enum, y)
+
+template <typename enumtype, typename storagetype>
+PX_INLINE PxFlags<enumtype, storagetype>::PxFlags(void)
+{
+	mBits = 0;
+}
+
+template <typename enumtype, typename storagetype>
+PX_INLINE PxFlags<enumtype, storagetype>::PxFlags(enumtype e)
+{
+	mBits = static_cast<storagetype>(e);
+}
+
+template <typename enumtype, typename storagetype>
+PX_INLINE PxFlags<enumtype, storagetype>::PxFlags(const PxFlags<enumtype, storagetype>& f)
+{
+	mBits = f.mBits;
+}
+
+template <typename enumtype, typename storagetype>
+PX_INLINE PxFlags<enumtype, storagetype>::PxFlags(storagetype b)
+{
+	mBits = b;
+}
+
+template <typename enumtype, typename storagetype>
+PX_INLINE bool PxFlags<enumtype, storagetype>::isSet(enumtype e) const
+{
+	return (mBits & static_cast<storagetype>(e)) == static_cast<storagetype>(e);
+}
+
+template <typename enumtype, typename storagetype>
+PX_INLINE PxFlags<enumtype, storagetype>& PxFlags<enumtype, storagetype>::set(enumtype e)
+{
+	mBits = static_cast<storagetype>(e);
+	return *this;
+}
+
+template <typename enumtype, typename storagetype>
+PX_INLINE bool PxFlags<enumtype, storagetype>::operator==(enumtype e) const
+{
+	return mBits == static_cast<storagetype>(e);
+}
+
+template <typename enumtype, typename storagetype>
+PX_INLINE bool PxFlags<enumtype, storagetype>::operator==(const PxFlags<enumtype, storagetype>& f) const
+{
+	return mBits == f.mBits;
+}
+
+template <typename enumtype, typename storagetype>
+PX_INLINE bool PxFlags<enumtype, storagetype>::operator==(bool b) const
+{
+	return bool(*this) == b;
+}
+
+template <typename enumtype, typename storagetype>
+PX_INLINE bool PxFlags<enumtype, storagetype>::operator!=(enumtype e) const
+{
+	return mBits != static_cast<storagetype>(e);
+}
+
+template <typename enumtype, typename storagetype>
+PX_INLINE bool PxFlags<enumtype, storagetype>::operator!=(const PxFlags<enumtype, storagetype>& f) const
+{
+	return mBits != f.mBits;
+}
+
+template <typename enumtype, typename storagetype>
+PX_INLINE PxFlags<enumtype, storagetype>& PxFlags<enumtype, storagetype>::operator=(enumtype e)
+{
+	mBits = static_cast<storagetype>(e);
+	return *this;
+}
+
+template <typename enumtype, typename storagetype>
+PX_INLINE PxFlags<enumtype, storagetype>& PxFlags<enumtype, storagetype>::operator=(const PxFlags<enumtype, storagetype>& f)
+{
+	mBits = f.mBits;
+	return *this;
+}
+
+template <typename enumtype, typename storagetype>
+PX_INLINE PxFlags<enumtype, storagetype>& PxFlags<enumtype, storagetype>::operator|=(enumtype e)
+{
+	mBits |= static_cast<storagetype>(e);
+	return *this;
+}
+
+template <typename enumtype, typename storagetype>
+PX_INLINE PxFlags<enumtype, storagetype>& PxFlags<enumtype, storagetype>::
+operator|=(const PxFlags<enumtype, storagetype>& f)
+{
+	mBits |= f.mBits;
+	return *this;
+}
+
+template <typename enumtype, typename storagetype>
+PX_INLINE PxFlags<enumtype, storagetype> PxFlags<enumtype, storagetype>::operator|(enumtype e) const
+{
+	PxFlags<enumtype, storagetype> out(*this);
+	out |= e;
+	return out;
+}
+
+template <typename enumtype, typename storagetype>
+PX_INLINE PxFlags<enumtype, storagetype> PxFlags<enumtype, storagetype>::
+operator|(const PxFlags<enumtype, storagetype>& f) const
+{
+	PxFlags<enumtype, storagetype> out(*this);
+	out |= f;
+	return out;
+}
+
+template <typename enumtype, typename storagetype>
+PX_INLINE PxFlags<enumtype, storagetype>& PxFlags<enumtype, storagetype>::operator&=(enumtype e)
+{
+	mBits &= static_cast<storagetype>(e);
+	return *this;
+}
+
+template <typename enumtype, typename storagetype>
+PX_INLINE PxFlags<enumtype, storagetype>& PxFlags<enumtype, storagetype>::
+operator&=(const PxFlags<enumtype, storagetype>& f)
+{
+	mBits &= f.mBits;
+	return *this;
+}
+
+template <typename enumtype, typename storagetype>
+PX_INLINE PxFlags<enumtype, storagetype> PxFlags<enumtype, storagetype>::operator&(enumtype e) const
+{
+	PxFlags<enumtype, storagetype> out = *this;
+	out.mBits &= static_cast<storagetype>(e);
+	return out;
+}
+
+template <typename enumtype, typename storagetype>
+PX_INLINE PxFlags<enumtype, storagetype> PxFlags<enumtype, storagetype>::
+operator&(const PxFlags<enumtype, storagetype>& f) const
+{
+	PxFlags<enumtype, storagetype> out = *this;
+	out.mBits &= f.mBits;
+	return out;
+}
+
+template <typename enumtype, typename storagetype>
+PX_INLINE PxFlags<enumtype, storagetype>& PxFlags<enumtype, storagetype>::operator^=(enumtype e)
+{
+	mBits ^= static_cast<storagetype>(e);
+	return *this;
+}
+
+template <typename enumtype, typename storagetype>
+PX_INLINE PxFlags<enumtype, storagetype>& PxFlags<enumtype, storagetype>::
+operator^=(const PxFlags<enumtype, storagetype>& f)
+{
+	mBits ^= f.mBits;
+	return *this;
+}
+
+template <typename enumtype, typename storagetype>
+PX_INLINE PxFlags<enumtype, storagetype> PxFlags<enumtype, storagetype>::operator^(enumtype e) const
+{
+	PxFlags<enumtype, storagetype> out = *this;
+	out.mBits ^= static_cast<storagetype>(e);
+	return out;
+}
+
+template <typename enumtype, typename storagetype>
+PX_INLINE PxFlags<enumtype, storagetype> PxFlags<enumtype, storagetype>::
+operator^(const PxFlags<enumtype, storagetype>& f) const
+{
+	PxFlags<enumtype, storagetype> out = *this;
+	out.mBits ^= f.mBits;
+	return out;
+}
+
+template <typename enumtype, typename storagetype>
+PX_INLINE PxFlags<enumtype, storagetype> PxFlags<enumtype, storagetype>::operator~(void) const
+{
+	PxFlags<enumtype, storagetype> out;
+	out.mBits = storagetype(~mBits);
+	return out;
+}
+
+template <typename enumtype, typename storagetype>
+PX_INLINE PxFlags<enumtype, storagetype>::operator bool(void) const
+{
+	return mBits ? true : false;
+}
+
+template <typename enumtype, typename storagetype>
+PX_INLINE PxFlags<enumtype, storagetype>::operator uint8_t(void) const
+{
+	return static_cast<uint8_t>(mBits);
+}
+
+template <typename enumtype, typename storagetype>
+PX_INLINE PxFlags<enumtype, storagetype>::operator uint16_t(void) const
+{
+	return static_cast<uint16_t>(mBits);
+}
+
+template <typename enumtype, typename storagetype>
+PX_INLINE PxFlags<enumtype, storagetype>::operator uint32_t(void) const
+{
+	return static_cast<uint32_t>(mBits);
+}
+
+template <typename enumtype, typename storagetype>
+PX_INLINE void PxFlags<enumtype, storagetype>::clear(enumtype e)
+{
+	mBits &= ~static_cast<storagetype>(e);
+}
+
+#endif // !WITH_PHYSX
diff --git a/Engine/Source/Runtime/PhysicsCore/Public/PhysicsInterfaceUtilsCore.h b/Engine/Source/Runtime/PhysicsCore/Public/PhysicsInterfaceUtilsCore.h
--- a/Engine/Source/Runtime/PhysicsCore/Public/PhysicsInterfaceUtilsCore.h
+++ b/Engine/Source/Runtime/PhysicsCore/Public/PhysicsInterfaceUtilsCore.h
@@ -6,15 +6,17 @@
 #include "PhysicsInterfaceDeclaresCore.h"
 
 #if WITH_PHYSX
-
 #include "PhysXPublicCore.h"
 #include "Containers/Union.h"
 #include "CollisionShape.h"
+#endif
 
+PHYSICSCORE_API FCollisionFilterData P2UFilterData(const PxFilterData& PFilterData);
+PHYSICSCORE_API PxFilterData U2PFilterData(const FCollisionFilterData& FilterData);
+
+#if WITH_PHYSX
 
 PHYSICSCORE_API PxShapeFlags BuildPhysXShapeFlags(FBodyCollisionFlags BodyCollisionFlags, bool bPhysicsStatic, bool bIsTriangleMesh);
-PHYSICSCORE_API PxFilterData U2PFilterData(const FCollisionFilterData& FilterData);
-PHYSICSCORE_API FCollisionFilterData P2UFilterData(const PxFilterData& PFilterData);
 PHYSICSCORE_API PxGeometryType::Enum U2PCollisionShapeType(ECollisionShapeType InUType);
 PHYSICSCORE_API ECollisionShapeType P2UCollisionShapeType(PxGeometryType::Enum InPType);
 
diff --git a/Engine/Source/Runtime/PhysicsCore/Public/SQCapture.h b/Engine/Source/Runtime/PhysicsCore/Public/SQCapture.h
--- a/Engine/Source/Runtime/PhysicsCore/Public/SQCapture.h
+++ b/Engine/Source/Runtime/PhysicsCore/Public/SQCapture.h
@@ -88,7 +88,7 @@ struct PHYSICSCORE_API FSQCapture
 	TUniquePtr<Chaos::FImplicitObject> ChaosOwnerObject;	//should be private, do not access directly
 	const Chaos::FImplicitObject* ChaosGeometry;
 	TUniquePtr<Chaos::FImplicitObject> SerializableChaosGeometry;
-#if WITH_PHYSX
+#if WITH_CHAOS
 	//for now just use physx hit buffer
 	ChaosInterface::FSQHitBuffer<ChaosInterface::FSweepHit> ChaosSweepBuffer;
 	TArray<ChaosInterface::FSweepHit> ChaosSweepTouches;
@@ -113,8 +113,10 @@ private:
 	TArray<uint8> GeomData;
 	TArray<uint8> HitData;
 
+#if WITH_PHYSX
 	void CreateChaosDataFromPhysX();
 	void CreateChaosFilterResults();
+#endif
 
 #if WITH_PHYSX
 
diff --git a/Engine/Source/Runtime/PhysicsSQ/Private/SQAccelerator.cpp b/Engine/Source/Runtime/PhysicsSQ/Private/SQAccelerator.cpp
--- a/Engine/Source/Runtime/PhysicsSQ/Private/SQAccelerator.cpp
+++ b/Engine/Source/Runtime/PhysicsSQ/Private/SQAccelerator.cpp
@@ -12,6 +12,7 @@
 #include "Experimental/SceneQueryChaosImp.h"
 #endif
 
+#include "PhysicsInterfaceUtilsCore.h"
 #include "ChaosInterfaceWrapperCore.h"
 #include "Chaos/CastingUtilities.h"
 #include "Chaos/ISpatialAcceleration.h"
@@ -97,10 +98,8 @@ struct TSQVisitor : public Chaos::ISpatialVisitor<TPayload, float>
 		, QueryFilterDataConcrete(P2UFilterData(QueryFilterData.data))
 		, QueryCallback(InQueryCallback)
 	{
-#if WITH_PHYSX
 		//#TODO - reimplement query flags alternative for Chaos
 		bAnyHit = QueryFilterData.flags & PxQueryFlag::eANY_HIT;
-#endif
 	}
 
 	TSQVisitor(const FTransform& InStartTM, const FVector& InDir, ChaosInterface::FSQHitBuffer<ChaosInterface::FSweepHit>& InHitBuffer, EHitFlags InOutputFlags,
@@ -117,10 +116,9 @@ struct TSQVisitor : public Chaos::ISpatialVisitor<TPayload, float>
 		, QueryCallback(InQueryCallback)
 		, StartTM(InStartTM)
 	{
-#if WITH_PHYSX
 		//#TODO - reimplement query flags alternative for Chaos
 		bAnyHit = QueryFilterData.flags & PxQueryFlag::eANY_HIT;
-#endif
+
 		//todo: check THitType is sweep
 	}
 
@@ -136,10 +134,9 @@ struct TSQVisitor : public Chaos::ISpatialVisitor<TPayload, float>
 		, QueryCallback(InQueryCallback)
 		, StartTM(InWorldTM)
 	{
-#if WITH_PHYSX
 		//#TODO - reimplement query flags alternative for Chaos
 		bAnyHit = QueryFilterData.flags & PxQueryFlag::eANY_HIT;
-#endif
+
 		//todo: check THitType is overlap
 	}
 
@@ -233,13 +230,8 @@ private:
 
 			//TODO: use gt particles directly
 			//#TODO alternative to px flags
-#if WITH_PHYSX
 			ECollisionQueryHitType HitType = QueryFilterData.flags & PxQueryFlag::ePREFILTER ? QueryCallback.PreFilter(QueryFilterDataConcrete, *Shape, *GeometryParticle) : ECollisionQueryHitType::Block;
-#else
-			//#TODO Chaos flag alternative
-			ensure(false);
-			ECollisionQueryHitType HitType = ECollisionQueryHitType::Block;
-#endif
+
 			if (HitType != ECollisionQueryHitType::None)
 			{
 				//QUICK_SCOPE_CYCLE_COUNTER(SQNarrow);
@@ -293,12 +285,9 @@ private:
 				{
 					//QUICK_SCOPE_CYCLE_COUNTER(SQNarrowHit);
 					FillHitHelper(Hit, Distance, WorldPosition, WorldNormal, FaceIdx, bComputeMTD);
-#if WITH_PHYSX
+
 					HitType = QueryFilterData.flags & PxQueryFlag::ePOSTFILTER ? QueryCallback.PostFilter(QueryFilterDataConcrete, Hit) : HitType;
-#else
-					//#TODO Chaos flag alternative
-					ensure(false);
-#endif
+
 					if (HitType != ECollisionQueryHitType::None)
 					{
 
