Fix compile

diff --git a/Engine/Source/Runtime/PhysicsCore/Private/SQCapture.cpp b/Engine/Source/Runtime/PhysicsCore/Private/SQCapture.cpp
--- a/Engine/Source/Runtime/PhysicsCore/Private/SQCapture.cpp
+++ b/Engine/Source/Runtime/PhysicsCore/Private/SQCapture.cpp
@@ -315,9 +315,11 @@ void FSQCapture::Serialize(Chaos::FChaosArchive& Ar)
 
 	if (bDiskDataIsChaos)
 	{
+#if WITH_CHAOS
 		SerializeChaosBuffers(Ar, Version, ChaosSweepBuffer);
 		SerializeChaosBuffers(Ar, Version, ChaosRaycastBuffer);
 		SerializeChaosBuffers(Ar, Version, ChaosOverlapBuffer);
+#endif // WITH_CHAOS
 
 		SerializeChaosActorToShapeHitsArray(Ar);
 		SerializeQueryFilterData(Ar, QueryFilterData);
@@ -366,7 +368,9 @@ void FSQCapture::EndCaptureChaosSweep(const ChaosInterface::FSQHitBuffer<ChaosIn
 	if (IsInGameThread())
 	{
 		check(SQType == ESQType::Sweep);
+#if WITH_CHAOS
 		ChaosSweepBuffer = Results;
+#endif // WITH_CHAOS
 	}
 }
 
@@ -392,7 +396,9 @@ void FSQCapture::EndCaptureChaosRaycast(const ChaosInterface::FSQHitBuffer<Chaos
 	if (IsInGameThread())
 	{
 		check(SQType == ESQType::Raycast);
+#if WITH_CHAOS
 		ChaosRaycastBuffer = Results;
+#endif // WITH_CHAOS
 	}
 }
 
@@ -417,7 +423,9 @@ void FSQCapture::EndCaptureChaosOverlap(const ChaosInterface::FSQHitBuffer<Chaos
 	if (IsInGameThread())
 	{
 		check(SQType == ESQType::Overlap);
+#if WITH_CHAOS
 		ChaosOverlapBuffer = Results;
+#endif // WITH_CHAOS
 	}
 }
 
@@ -813,9 +821,11 @@ void FSQCapture::CreateChaosDataFromPhysX()
 		ChaosOwnerObject = MoveTemp(TransformedObj);
 	}
 
+#if WITH_CHAOS
 	PhysXToChaosBufferData(ChaosRaycastBuffer, PhysXRaycastBuffer, PhysSerializer);
 	PhysXToChaosBufferData(ChaosSweepBuffer, PhysXSweepBuffer, PhysSerializer);
 	PhysXToChaosBufferData(ChaosOverlapBuffer, PhysXOverlapBuffer, PhysSerializer);
+#endif // WITH_CHAOS
 
 	CreateChaosFilterResults();
 	bChaosDataReady = true;
