Fix build error with PhysX and Chaos, and IWYU.

diff --git a/Engine/Source/Runtime/PhysicsCore/Private/PhysicsInitialization.cpp b/Engine/Source/Runtime/PhysicsCore/Private/PhysicsInitialization.cpp
--- a/Engine/Source/Runtime/PhysicsCore/Private/PhysicsInitialization.cpp
+++ b/Engine/Source/Runtime/PhysicsCore/Private/PhysicsInitialization.cpp
@@ -9,6 +9,7 @@
 #include "IPhysXCooking.h"
 #include "Modules/ModuleManager.h"
 #include "Misc/Paths.h"
+#include "Core/Public/HAL/IConsoleManager.h"
 
 #ifndef APEX_STATICALLY_LINKED
 #define APEX_STATICALLY_LINKED	0
diff --git a/Engine/Source/Runtime/PhysicsCore/Public/PhysicsInterfaceTypesCore.h b/Engine/Source/Runtime/PhysicsCore/Public/PhysicsInterfaceTypesCore.h
--- a/Engine/Source/Runtime/PhysicsCore/Public/PhysicsInterfaceTypesCore.h
+++ b/Engine/Source/Runtime/PhysicsCore/Public/PhysicsInterfaceTypesCore.h
@@ -145,7 +145,7 @@ static bool CalcMeshNegScaleCompensation(const FVector& InScale3D, FTransform& O
 
 
 //TODO: Reimplement types in chaos
-#if !PHYSICS_INTERFACE_PHYSX
+#if !PHYSICS_INTERFACE_PHYSX && WITH_CHAOS
 
 const uint32 AggregateMaxSize = 128;
 
