Make PhysicsCore not include public dependency on PhysX when Chaos is enabled. Fixes NonUnity build.

diff --git a/Engine/Source/Runtime/PhysicsCore/PhysicsCore.Build.cs b/Engine/Source/Runtime/PhysicsCore/PhysicsCore.Build.cs
--- a/Engine/Source/Runtime/PhysicsCore/PhysicsCore.Build.cs
+++ b/Engine/Source/Runtime/PhysicsCore/PhysicsCore.Build.cs
@@ -17,7 +17,7 @@ public class PhysicsCore: ModuleRules
 
 		SetupModulePhysicsSupport(Target); 
 
-		if(Target.bCompilePhysX)
+		if(Target.bCompilePhysX && Target.bCompileChaos == false && Target.bUseChaos == false)
         {
 			// Not ideal but as this module publicly exposes PhysX types
 			// to other modules when PhysX is enabled it requires that its
