Resubmit - PhysicsCore no longer includes PhysX as public dependency when Chaos is enabled.
Fixes NonUnity build.

diff --git a/Engine/Source/Runtime/PhysicsCore/PhysicsCore.Build.cs b/Engine/Source/Runtime/PhysicsCore/PhysicsCore.Build.cs
--- a/Engine/Source/Runtime/PhysicsCore/PhysicsCore.Build.cs
+++ b/Engine/Source/Runtime/PhysicsCore/PhysicsCore.Build.cs
@@ -15,19 +15,22 @@ public class PhysicsCore: ModuleRules
 			}
 		);
 
-		SetupModulePhysicsSupport(Target); 
+		SetupModulePhysicsSupport(Target);
 
-		if(Target.bCompilePhysX)
+        if (Target.bCompileChaos == false && Target.bUseChaos == false)
         {
-			// Not ideal but as this module publicly exposes PhysX types
-			// to other modules when PhysX is enabled it requires that its
-			// public files have access to PhysX includes
-            PublicDependencyModuleNames.Add("PhysX");
-        }
+            if (Target.bCompilePhysX)
+            {
+                // Not ideal but as this module publicly exposes PhysX types
+                // to other modules when PhysX is enabled it requires that its
+                // public files have access to PhysX includes
+                PublicDependencyModuleNames.Add("PhysX");
+            }
 
-		if(Target.bCompileAPEX)
-        {
-            PublicDependencyModuleNames.Add("APEX");
+            if (Target.bCompileAPEX)
+            {
+                PublicDependencyModuleNames.Add("APEX");
+            }
         }
 	}
 }
