Rename duplicated types from PhysX when removing dependancy. Fixes NonUnity build ambiguous type errors.

diff --git a/Engine/Source/Runtime/Engine/Classes/PhysicalMaterials/PhysicalMaterial.h b/Engine/Source/Runtime/Engine/Classes/PhysicalMaterials/PhysicalMaterial.h
--- a/Engine/Source/Runtime/Engine/Classes/PhysicalMaterials/PhysicalMaterial.h
+++ b/Engine/Source/Runtime/Engine/Classes/PhysicalMaterials/PhysicalMaterial.h
@@ -7,7 +7,6 @@
 #include "UObject/Object.h"
 #include "Engine/EngineTypes.h"
 #include "EngineDefines.h"
-#include "PhysxUserData.h"
 #include "Vehicles/TireType.h"
 #include "PhysicsEngine/PhysicsSettingsEnums.h"
 #include "Physics/PhysicsInterfaceCore.h"
@@ -132,7 +131,7 @@ public:
 
 	FPhysicsMaterialHandle MaterialHandle;
 
-	FPhysxUserData PhysxUserData;
+	FChaosUserData UserData;
 
 	//~ Begin UObject Interface
 #if WITH_EDITOR
diff --git a/Engine/Source/Runtime/Engine/Classes/PhysicsEngine/ConstraintInstance.h b/Engine/Source/Runtime/Engine/Classes/PhysicsEngine/ConstraintInstance.h
--- a/Engine/Source/Runtime/Engine/Classes/PhysicsEngine/ConstraintInstance.h
+++ b/Engine/Source/Runtime/Engine/Classes/PhysicsEngine/ConstraintInstance.h
@@ -8,7 +8,6 @@
 #include "PhysicsEngine/ConstraintTypes.h"
 #include "PhysicsEngine/ConstraintDrives.h"
 #include "Physics/PhysicsInterfaceCore.h"
-#include "PhysxUserData.h"
 #include "ConstraintInstance.generated.h"
 
 class FMaterialRenderProxy;
@@ -182,7 +181,7 @@ public:
 	/** Get underlying physics engine constraint */
 	const FPhysicsConstraintHandle& GetPhysicsConstraintRef() const;
 
-	FPhysxUserData PhysxUserData;
+	FChaosUserData UserData;
 
 private:
 	/** The component scale passed in during initialization*/
diff --git a/Engine/Source/Runtime/Engine/Classes/PhysicsEngine/ShapeElem.h b/Engine/Source/Runtime/Engine/Classes/PhysicsEngine/ShapeElem.h
--- a/Engine/Source/Runtime/Engine/Classes/PhysicsEngine/ShapeElem.h
+++ b/Engine/Source/Runtime/Engine/Classes/PhysicsEngine/ShapeElem.h
@@ -67,7 +67,7 @@ struct FKShapeElem
 		return (T*)this;
 	}
 
-	const FPhysxUserData* GetUserData() const { FPhysxUserData::Set<FKShapeElem>((void*)&UserData, const_cast<FKShapeElem*>(this));  return &UserData; }
+	const FUserData* GetUserData() const { FUserData::Set<FKShapeElem>((void*)&UserData, const_cast<FKShapeElem*>(this));  return &UserData; }
 
 	ENGINE_API static EAggCollisionShape::Type StaticShapeType;
 
@@ -111,5 +111,5 @@ private:
 	UPROPERTY(Category=Shape, EditAnywhere)
 	uint8 bContributeToMass : 1;
 
-	FPhysxUserData UserData;
+	FUserData UserData;
 };
diff --git a/Engine/Source/Runtime/Engine/Private/PhysicsEngine/BodyInstance.cpp b/Engine/Source/Runtime/Engine/Private/PhysicsEngine/BodyInstance.cpp
--- a/Engine/Source/Runtime/Engine/Private/PhysicsEngine/BodyInstance.cpp
+++ b/Engine/Source/Runtime/Engine/Private/PhysicsEngine/BodyInstance.cpp
@@ -1783,7 +1783,7 @@ bool FBodyInstance::UpdateBodyScale(const FVector& InScale3D, bool bForceUpdate)
 				CHAOS_ENSURE(!IsInstanced(GeomType));
 			}
 
-			FKShapeElem* ShapeElem = FPhysxUserData::Get<FKShapeElem>(FPhysicsInterface::GetUserData(ShapeHandle));
+			FKShapeElem* ShapeElem = FChaosUserData::Get<FKShapeElem>(FPhysicsInterface::GetUserData(ShapeHandle));
 
 			switch (GeomType)
 			{
@@ -3056,7 +3056,7 @@ void FBodyInstance::UpdateMassProperties()
 				for (int32 ShapeIdx = Shapes.Num() - 1; ShapeIdx >= 0; --ShapeIdx)
 				{
 					const FPhysicsShapeHandle& Shape = Shapes[ShapeIdx];
-					const FKShapeElem* ShapeElem = FPhysxUserData::Get<FKShapeElem>(FPhysicsInterface::GetUserData(Shape));
+					const FKShapeElem* ShapeElem = FChaosUserData::Get<FKShapeElem>(FPhysicsInterface::GetUserData(Shape));
 					bool bIsTriangleMesh = FPhysicsInterface::GetShapeType(Shape) == ECollisionShapeType::Trimesh;
 					bool bHasNoMass = ShapeElem && !ShapeElem->GetContributeToMass();
 					if (bIsTriangleMesh || bHasNoMass)
diff --git a/Engine/Source/Runtime/Engine/Private/PhysicsEngine/ConstraintInstance.cpp b/Engine/Source/Runtime/Engine/Private/PhysicsEngine/ConstraintInstance.cpp
--- a/Engine/Source/Runtime/Engine/Private/PhysicsEngine/ConstraintInstance.cpp
+++ b/Engine/Source/Runtime/Engine/Private/PhysicsEngine/ConstraintInstance.cpp
@@ -181,9 +181,7 @@ FConstraintInstance::FConstraintInstance()
 	, AngularRotationOffset(ForceInitToZero)
 	, bScaleLinearLimits(true)
 	, AverageMass(0.f)
-#if WITH_PHYSX
-	, PhysxUserData(this)
-#endif
+	, UserData(this)
 	, LastKnownScale(1.f)
 #if WITH_EDITORONLY_DATA
 	, bDisableCollision_DEPRECATED(false)
@@ -378,7 +376,7 @@ bool FConstraintInstance::CreateJoint_AssumesLocked(const FPhysicsActorHandle& I
 		return false;
 	}
 
-	FPhysicsInterface::SetConstraintUserData(ConstraintHandle, &PhysxUserData);
+	FPhysicsInterface::SetConstraintUserData(ConstraintHandle, &UserData);
 
 	return true;
 }
@@ -448,7 +446,7 @@ void FConstraintInstance::InitConstraint_AssumesLocked(const FPhysicsActorHandle
 	OnConstraintBrokenDelegate = InConstraintBrokenDelegate;
 	LastKnownScale = InScale;
 
-	PhysxUserData = FPhysxUserData(this);
+	UserData = FChaosUserData(this);
 
 	// if there's already a constraint, get rid of it first
 	if (ConstraintHandle.IsValid())
diff --git a/Engine/Source/Runtime/Engine/Private/PhysicsEngine/Experimental/ChaosInterfaceWrapper.cpp b/Engine/Source/Runtime/Engine/Private/PhysicsEngine/Experimental/ChaosInterfaceWrapper.cpp
--- a/Engine/Source/Runtime/Engine/Private/PhysicsEngine/Experimental/ChaosInterfaceWrapper.cpp
+++ b/Engine/Source/Runtime/Engine/Private/PhysicsEngine/Experimental/ChaosInterfaceWrapper.cpp
@@ -11,13 +11,13 @@ namespace ChaosInterface
 	FBodyInstance* GetUserData(const Chaos::TGeometryParticle<float,3>& Actor)
 	{
 		void* UserData = Actor.UserData();
-		return UserData ? FPhysxUserData::Get<FBodyInstance>(Actor.UserData()) : nullptr;
+		return UserData ? FChaosUserData::Get<FBodyInstance>(Actor.UserData()) : nullptr;
 	}
 
 	UPhysicalMaterial* GetUserData(const Chaos::FChaosPhysicsMaterial& Material)
 	{
 		void* UserData = Material.UserData;
-		return UserData ? FPhysxUserData::Get<UPhysicalMaterial>(UserData) : nullptr;
+		return UserData ? FChaosUserData::Get<UPhysicalMaterial>(UserData) : nullptr;
 	}
 
 #if WITH_CHAOS
diff --git a/Engine/Source/Runtime/Engine/Private/PhysicsEngine/PhysicalMaterial.cpp b/Engine/Source/Runtime/Engine/Private/PhysicsEngine/PhysicalMaterial.cpp
--- a/Engine/Source/Runtime/Engine/Private/PhysicsEngine/PhysicalMaterial.cpp
+++ b/Engine/Source/Runtime/Engine/Private/PhysicsEngine/PhysicalMaterial.cpp
@@ -34,9 +34,7 @@ UPhysicalMaterial::UPhysicalMaterial(const FObjectInitializer& ObjectInitializer
 	DestructibleDamageThresholdScale = 1.0f;
 	TireFrictionScale = 1.0f;
 	bOverrideFrictionCombineMode = false;
-#if PHYSICS_INTERFACE_PHYSX
-	PhysxUserData = FPhysxUserData(this);
-#endif
+	UserData = FChaosUserData(this);
 }
 
 #if WITH_EDITOR
@@ -89,7 +87,7 @@ FPhysicsMaterialHandle& UPhysicalMaterial::GetPhysicsMaterial()
 		MaterialHandle = FPhysicsInterface::CreateMaterial(this);
 		check(MaterialHandle.IsValid());
 
-		FPhysicsInterface::SetUserData(MaterialHandle, &PhysxUserData);
+		FPhysicsInterface::SetUserData(MaterialHandle, &UserData);
 		FPhysicsInterface::UpdateMaterial(MaterialHandle, this);
 	}
 
diff --git a/Engine/Source/Runtime/Engine/Public/Physics/Experimental/ChaosInterfaceWrapper.h b/Engine/Source/Runtime/Engine/Public/Physics/Experimental/ChaosInterfaceWrapper.h
--- a/Engine/Source/Runtime/Engine/Public/Physics/Experimental/ChaosInterfaceWrapper.h
+++ b/Engine/Source/Runtime/Engine/Public/Physics/Experimental/ChaosInterfaceWrapper.h
@@ -35,7 +35,11 @@ struct FScopedSceneReadLock
 
 inline FQueryFilterData MakeQueryFilterData(const FCollisionFilterData& FilterData, EQueryFlags QueryFlags, const FCollisionQueryParams& Params)
 {
+#if PHYSICS_INTERFACE_PHYSX
 	return PxQueryFilterData(U2PFilterData(FilterData), U2PQueryFlags(QueryFlags));
+#else
+	return FChaosQueryFilterData(U2CFilterData(FilterData), U2CQueryFlags(QueryFlags));
+#endif
 }
 
 FBodyInstance* GetUserData(const Chaos::TGeometryParticle<float, 3>& Actor);
diff --git a/Engine/Source/Runtime/Engine/Public/PhysxUserData.h b/Engine/Source/Runtime/Engine/Public/PhysxUserData.h
--- a/Engine/Source/Runtime/Engine/Public/PhysxUserData.h
+++ b/Engine/Source/Runtime/Engine/Public/PhysxUserData.h
@@ -1,7 +1,6 @@
 // Copyright Epic Games, Inc. All Rights Reserved.
 #pragma once
 
-#if PHYSICS_INTERFACE_PHYSX
 #include "CoreMinimal.h"
 #include "EngineDefines.h"
 #include "Physics/PhysicsInterfaceDeclares.h"
@@ -71,5 +70,4 @@ template <> FORCEINLINE void FPhysxUserData::Set(void* UserData, FPhysScene* Pay
 template <> FORCEINLINE void FPhysxUserData::Set(void* UserData, FConstraintInstance* Payload)		{ check(UserData); ((FPhysxUserData*)UserData)->Type = EPhysxUserDataType::ConstraintInstance; ((FPhysxUserData*)UserData)->Payload = Payload; }
 template <> FORCEINLINE void FPhysxUserData::Set(void* UserData, UPrimitiveComponent* Payload)		{ check(UserData); ((FPhysxUserData*)UserData)->Type = EPhysxUserDataType::PrimitiveComponent; ((FPhysxUserData*)UserData)->Payload = Payload; }
 template <> FORCEINLINE void FPhysxUserData::Set(void* UserData, FKShapeElem* Payload)	{ check(UserData); ((FPhysxUserData*)UserData)->Type = EPhysxUserDataType::AggShape; ((FPhysxUserData*)UserData)->Payload = Payload; }
-template <> FORCEINLINE void FPhysxUserData::Set(void* UserData, FCustomPhysXPayload* Payload) { check(UserData); ((FPhysxUserData*)UserData)->Type = EPhysxUserDataType::CustomPayload; ((FPhysxUserData*)UserData)->Payload = Payload; }
-#endif 
\ No newline at end of file
+template <> FORCEINLINE void FPhysxUserData::Set(void* UserData, FCustomPhysXPayload* Payload) { check(UserData); ((FPhysxUserData*)UserData)->Type = EPhysxUserDataType::CustomPayload; ((FPhysxUserData*)UserData)->Payload = Payload; }
\ No newline at end of file
diff --git a/Engine/Source/Runtime/PhysicsCore/Private/PhysicsInterfaceUtilsCore.cpp b/Engine/Source/Runtime/PhysicsCore/Private/PhysicsInterfaceUtilsCore.cpp
--- a/Engine/Source/Runtime/PhysicsCore/Private/PhysicsInterfaceUtilsCore.cpp
+++ b/Engine/Source/Runtime/PhysicsCore/Private/PhysicsInterfaceUtilsCore.cpp
@@ -6,6 +6,24 @@
 
 #include "Chaos/ParticleHandle.h"
 
+
+FCollisionFilterData C2UFilterData(const FChaosFilterData& PFilterData)
+{
+	FCollisionFilterData FilterData;
+	FilterData.Word0 = PFilterData.word0;
+	FilterData.Word1 = PFilterData.word1;
+	FilterData.Word2 = PFilterData.word2;
+	FilterData.Word3 = PFilterData.word3;
+	return FilterData;
+}
+
+FChaosFilterData U2CFilterData(const FCollisionFilterData& FilterData)
+{
+	return FChaosFilterData(FilterData.Word0, FilterData.Word1, FilterData.Word2, FilterData.Word3);
+}
+
+#if PHYSICS_INTERFACE_PHYSX
+
 FCollisionFilterData P2UFilterData(const PxFilterData& PFilterData)
 {
 	FCollisionFilterData FilterData;
@@ -21,7 +39,6 @@ PxFilterData U2PFilterData(const FCollisionFilterData& FilterData)
 	return PxFilterData(FilterData.Word0, FilterData.Word1, FilterData.Word2, FilterData.Word3);
 }
 
-#if PHYSICS_INTERFACE_PHYSX
 
 PxShapeFlags BuildPhysXShapeFlags(FBodyCollisionFlags BodyCollisionFlags, bool bPhysicsStatic, bool bIsTriangleMesh)
 {
diff --git a/Engine/Source/Runtime/PhysicsCore/Private/SQCapture.cpp b/Engine/Source/Runtime/PhysicsCore/Private/SQCapture.cpp
--- a/Engine/Source/Runtime/PhysicsCore/Private/SQCapture.cpp
+++ b/Engine/Source/Runtime/PhysicsCore/Private/SQCapture.cpp
@@ -4,6 +4,7 @@
 
 #include "Chaos/ImplicitObject.h"
 #include "Chaos/PBDRigidsEvolutionGBF.h"
+#include "PhysicsInterfaceTypesCore.h"
 
 #include "PhysXSupportCore.h"
 #include "PhysicsPublicCore.h"
@@ -184,7 +185,11 @@ void SerializeQueryFilterData(FArchive& Ar, FQueryFilterData& QueryFilterData)
 	Ar << QueryFilterData.data.word3;
 	uint16 Flags = QueryFilterData.flags;
 	Ar << Flags;
+#if PHYSICS_INTERFACE_PHYSX
 	QueryFilterData.flags = (PxQueryFlags)Flags;
+#else
+	QueryFilterData.flags = (FChaosQueryFlags)Flags;
+#endif
 	Ar << QueryFilterData.clientId;
 }
 
diff --git a/Engine/Source/Runtime/PhysicsCore/Public/ChaosInterfaceWrapperCore.h b/Engine/Source/Runtime/PhysicsCore/Public/ChaosInterfaceWrapperCore.h
--- a/Engine/Source/Runtime/PhysicsCore/Public/ChaosInterfaceWrapperCore.h
+++ b/Engine/Source/Runtime/PhysicsCore/Public/ChaosInterfaceWrapperCore.h
@@ -28,7 +28,11 @@ struct FDummyPhysActor {};
 template<typename DummyT>
 struct FDummyCallback {};
 
+#if PHYSICS_INTERFACE_PHYSX
 using FQueryFilterData = PxQueryFilterData;
+#elif WITH_CHAOS
+using FQueryFilterData = FChaosQueryFilterData;
+#endif
 
 /** We use this struct so that if no conversion is needed in another API, we can avoid the copy (if we think that's critical) */
 struct FPhysicsRaycastInputAdapater
diff --git a/Engine/Source/Runtime/PhysicsCore/Public/PhysicsInterfaceTypesCore.h b/Engine/Source/Runtime/PhysicsCore/Public/PhysicsInterfaceTypesCore.h
--- a/Engine/Source/Runtime/PhysicsCore/Public/PhysicsInterfaceTypesCore.h
+++ b/Engine/Source/Runtime/PhysicsCore/Public/PhysicsInterfaceTypesCore.h
@@ -3,7 +3,7 @@
 #pragma once
 #include "CoreMinimal.h"
 #include "PhysicsInterfaceDeclaresCore.h"
-
+#include "PhysicsInterfaceWrapperShared.h"
 #include "Chaos/CollisionFilterData.h"
 
 struct FBodyInstance;
@@ -144,11 +144,12 @@ static bool CalcMeshNegScaleCompensation(const FVector& InScale3D, FTransform& O
 }
 
 
-//TODO: Reimplement types in chaos
-#if !PHYSICS_INTERFACE_PHYSX && WITH_CHAOS
+// TODO: Fixup types, these are more or less the PhysX types renamed as a temporary solution.
+// Probably should move to different header as well.
 
+#if !PHYSICS_INTERFACE_PHYSX
 const uint32 AggregateMaxSize = 128;
-
+#endif
 
 class UPhysicalMaterial;
 class UPrimitiveComponent;
@@ -158,10 +159,9 @@ struct FKShapeElem;
 
 /** Forward declarations */
 struct FKShapeElem;
-struct FCustomPhysXPayload;
+struct FCustomChaosPayload;
 
-/** PhysX user data type*/
-namespace EPhysxUserDataType
+namespace EChaosUserDataType
 {
 	enum Type
 	{
@@ -176,123 +176,84 @@ namespace EPhysxUserDataType
 	};
 };
 
-/** PhysX user data */
-struct FPhysxUserData
+
+struct FChaosUserData
 {
 protected:
-	EPhysxUserDataType::Type	Type;
+	EChaosUserDataType::Type	Type;
 	void*						Payload;
 
 public:
-	FPhysxUserData()									:Type(EPhysxUserDataType::Invalid), Payload(nullptr) {}
-	FPhysxUserData(FBodyInstance* InPayload)			:Type(EPhysxUserDataType::BodyInstance), Payload(InPayload) {}
-	FPhysxUserData(UPhysicalMaterial* InPayload)		:Type(EPhysxUserDataType::PhysicalMaterial), Payload(InPayload) {}
-	FPhysxUserData(FPhysScene* InPayload)			    :Type(EPhysxUserDataType::PhysScene), Payload(InPayload) {}
-	FPhysxUserData(FConstraintInstance* InPayload)		:Type(EPhysxUserDataType::ConstraintInstance), Payload(InPayload) {}
-	FPhysxUserData(UPrimitiveComponent* InPayload)		:Type(EPhysxUserDataType::PrimitiveComponent), Payload(InPayload) {}
-	FPhysxUserData(FKShapeElem* InPayload)				:Type(EPhysxUserDataType::AggShape), Payload(InPayload) {}
-	FPhysxUserData(FCustomPhysXPayload* InPayload)		:Type(EPhysxUserDataType::CustomPayload), Payload(InPayload) {}
+	FChaosUserData()									:Type(EChaosUserDataType::Invalid), Payload(nullptr) {}
+	FChaosUserData(FBodyInstance* InPayload)			:Type(EChaosUserDataType::BodyInstance), Payload(InPayload) {}
+	FChaosUserData(UPhysicalMaterial* InPayload)		:Type(EChaosUserDataType::PhysicalMaterial), Payload(InPayload) {}
+	FChaosUserData(FPhysScene* InPayload)			    :Type(EChaosUserDataType::PhysScene), Payload(InPayload) {}
+	FChaosUserData(FConstraintInstance* InPayload)		:Type(EChaosUserDataType::ConstraintInstance), Payload(InPayload) {}
+	FChaosUserData(UPrimitiveComponent* InPayload)		:Type(EChaosUserDataType::PrimitiveComponent), Payload(InPayload) {}
+	FChaosUserData(FKShapeElem* InPayload)				:Type(EChaosUserDataType::AggShape), Payload(InPayload) {}
+	FChaosUserData(FCustomChaosPayload* InPayload)		:Type(EChaosUserDataType::CustomPayload), Payload(InPayload) {}
 	
 	template <class T> static T* Get(void* UserData);
 	template <class T> static void Set(void* UserData, T* Payload);
 
 	//helper function to determine if userData is garbage (maybe dangling pointer)
-	static bool IsGarbage(void* UserData){ return ((FPhysxUserData*)UserData)->Type < EPhysxUserDataType::Invalid || ((FPhysxUserData*)UserData)->Type > EPhysxUserDataType::CustomPayload; }
+	static bool IsGarbage(void* UserData){ return ((FChaosUserData*)UserData)->Type < EChaosUserDataType::Invalid || ((FChaosUserData*)UserData)->Type > EChaosUserDataType::CustomPayload; }
 };
 
-template <> FORCEINLINE FBodyInstance* FPhysxUserData::Get(void* UserData)			{ if (!UserData || ((FPhysxUserData*)UserData)->Type != EPhysxUserDataType::BodyInstance) { return nullptr; } return (FBodyInstance*)((FPhysxUserData*)UserData)->Payload; }
-template <> FORCEINLINE UPhysicalMaterial* FPhysxUserData::Get(void* UserData)		{ if (!UserData || ((FPhysxUserData*)UserData)->Type != EPhysxUserDataType::PhysicalMaterial) { return nullptr; } return (UPhysicalMaterial*)((FPhysxUserData*)UserData)->Payload; }
-template <> FORCEINLINE FPhysScene* FPhysxUserData::Get(void* UserData)				{ if (!UserData || ((FPhysxUserData*)UserData)->Type != EPhysxUserDataType::PhysScene) { return nullptr; }return (FPhysScene*)((FPhysxUserData*)UserData)->Payload; }
-template <> FORCEINLINE FConstraintInstance* FPhysxUserData::Get(void* UserData)	{ if (!UserData || ((FPhysxUserData*)UserData)->Type != EPhysxUserDataType::ConstraintInstance) { return nullptr; } return (FConstraintInstance*)((FPhysxUserData*)UserData)->Payload; }
-template <> FORCEINLINE UPrimitiveComponent* FPhysxUserData::Get(void* UserData)	{ if (!UserData || ((FPhysxUserData*)UserData)->Type != EPhysxUserDataType::PrimitiveComponent) { return nullptr; } return (UPrimitiveComponent*)((FPhysxUserData*)UserData)->Payload; }
-template <> FORCEINLINE FKShapeElem* FPhysxUserData::Get(void* UserData)	{ if (!UserData || ((FPhysxUserData*)UserData)->Type != EPhysxUserDataType::AggShape) { return nullptr; } return (FKShapeElem*)((FPhysxUserData*)UserData)->Payload; }
-template <> FORCEINLINE FCustomPhysXPayload* FPhysxUserData::Get(void* UserData) { if (!UserData || ((FPhysxUserData*)UserData)->Type != EPhysxUserDataType::CustomPayload) { return nullptr; } return (FCustomPhysXPayload*)((FPhysxUserData*)UserData)->Payload; }
-
-template <> FORCEINLINE void FPhysxUserData::Set(void* UserData, FBodyInstance* Payload)			{ check(UserData); ((FPhysxUserData*)UserData)->Type = EPhysxUserDataType::BodyInstance; ((FPhysxUserData*)UserData)->Payload = Payload; }
-template <> FORCEINLINE void FPhysxUserData::Set(void* UserData, UPhysicalMaterial* Payload)		{ check(UserData); ((FPhysxUserData*)UserData)->Type = EPhysxUserDataType::PhysicalMaterial; ((FPhysxUserData*)UserData)->Payload = Payload; }
-template <> FORCEINLINE void FPhysxUserData::Set(void* UserData, FPhysScene* Payload)				{ check(UserData); ((FPhysxUserData*)UserData)->Type = EPhysxUserDataType::PhysScene; ((FPhysxUserData*)UserData)->Payload = Payload; }
-template <> FORCEINLINE void FPhysxUserData::Set(void* UserData, FConstraintInstance* Payload)		{ check(UserData); ((FPhysxUserData*)UserData)->Type = EPhysxUserDataType::ConstraintInstance; ((FPhysxUserData*)UserData)->Payload = Payload; }
-template <> FORCEINLINE void FPhysxUserData::Set(void* UserData, UPrimitiveComponent* Payload)		{ check(UserData); ((FPhysxUserData*)UserData)->Type = EPhysxUserDataType::PrimitiveComponent; ((FPhysxUserData*)UserData)->Payload = Payload; }
-template <> FORCEINLINE void FPhysxUserData::Set(void* UserData, FKShapeElem* Payload)	{ check(UserData); ((FPhysxUserData*)UserData)->Type = EPhysxUserDataType::AggShape; ((FPhysxUserData*)UserData)->Payload = Payload; }
-template <> FORCEINLINE void FPhysxUserData::Set(void* UserData, FCustomPhysXPayload* Payload) { check(UserData); ((FPhysxUserData*)UserData)->Type = EPhysxUserDataType::CustomPayload; ((FPhysxUserData*)UserData)->Payload = Payload; }
-
-/** enum for empty constructor tag*/
-enum PxEMPTY
-{
-	PxEmpty
-};
+#if PHYSICS_INTERFACE_PHYSX
+using FUserData = FPhysxUserData;
+#else
+using FUserData = FChaosUserData;
+#endif
 
-#define PX_INLINE inline
-
-typedef int64_t PxI64;
-typedef uint64_t PxU64;
-typedef int32_t PxI32;
-typedef uint32_t PxU32;
-typedef int16_t PxI16;
-typedef uint16_t PxU16;
-typedef int8_t PxI8;
-typedef uint8_t PxU8;
-typedef float PxF32;
-typedef double PxF64;
-typedef float PxReal;
-
-struct PxFilterData
-{
-	//= ATTENTION! =====================================================================================
-	// Changing the data layout of this class breaks the binary serialization format.  See comments for 
-	// PX_BINARY_SERIAL_VERSION.  If a modification is required, please adjust the getBinaryMetaData 
-	// function.  If the modification is made on a custom branch, please change PX_BINARY_SERIAL_VERSION
-	// accordingly.
-	//==================================================================================================
-
-	PX_INLINE PxFilterData(const PxEMPTY)
-	{
-	}
 
-	/**
-	\brief Default constructor.
-	*/
-	PX_INLINE PxFilterData()
+template <> FORCEINLINE FBodyInstance* FChaosUserData::Get(void* UserData)			{ if (!UserData || ((FChaosUserData*)UserData)->Type != EChaosUserDataType::BodyInstance) { return nullptr; } return (FBodyInstance*)((FChaosUserData*)UserData)->Payload; }
+template <> FORCEINLINE UPhysicalMaterial* FChaosUserData::Get(void* UserData)		{ if (!UserData || ((FChaosUserData*)UserData)->Type != EChaosUserDataType::PhysicalMaterial) { return nullptr; } return (UPhysicalMaterial*)((FChaosUserData*)UserData)->Payload; }
+template <> FORCEINLINE FPhysScene* FChaosUserData::Get(void* UserData)				{ if (!UserData || ((FChaosUserData*)UserData)->Type != EChaosUserDataType::PhysScene) { return nullptr; }return (FPhysScene*)((FChaosUserData*)UserData)->Payload; }
+template <> FORCEINLINE FConstraintInstance* FChaosUserData::Get(void* UserData)	{ if (!UserData || ((FChaosUserData*)UserData)->Type != EChaosUserDataType::ConstraintInstance) { return nullptr; } return (FConstraintInstance*)((FChaosUserData*)UserData)->Payload; }
+template <> FORCEINLINE UPrimitiveComponent* FChaosUserData::Get(void* UserData)	{ if (!UserData || ((FChaosUserData*)UserData)->Type != EChaosUserDataType::PrimitiveComponent) { return nullptr; } return (UPrimitiveComponent*)((FChaosUserData*)UserData)->Payload; }
+template <> FORCEINLINE FKShapeElem* FChaosUserData::Get(void* UserData)	{ if (!UserData || ((FChaosUserData*)UserData)->Type != EChaosUserDataType::AggShape) { return nullptr; } return (FKShapeElem*)((FChaosUserData*)UserData)->Payload; }
+template <> FORCEINLINE FCustomChaosPayload* FChaosUserData::Get(void* UserData) { if (!UserData || ((FChaosUserData*)UserData)->Type != EChaosUserDataType::CustomPayload) { return nullptr; } return (FCustomChaosPayload*)((FChaosUserData*)UserData)->Payload; }
+
+template <> FORCEINLINE void FChaosUserData::Set(void* UserData, FBodyInstance* Payload)			{ check(UserData); ((FChaosUserData*)UserData)->Type = EChaosUserDataType::BodyInstance; ((FChaosUserData*)UserData)->Payload = Payload; }
+template <> FORCEINLINE void FChaosUserData::Set(void* UserData, UPhysicalMaterial* Payload)		{ check(UserData); ((FChaosUserData*)UserData)->Type = EChaosUserDataType::PhysicalMaterial; ((FChaosUserData*)UserData)->Payload = Payload; }
+template <> FORCEINLINE void FChaosUserData::Set(void* UserData, FPhysScene* Payload)				{ check(UserData); ((FChaosUserData*)UserData)->Type = EChaosUserDataType::PhysScene; ((FChaosUserData*)UserData)->Payload = Payload; }
+template <> FORCEINLINE void FChaosUserData::Set(void* UserData, FConstraintInstance* Payload)		{ check(UserData); ((FChaosUserData*)UserData)->Type = EChaosUserDataType::ConstraintInstance; ((FChaosUserData*)UserData)->Payload = Payload; }
+template <> FORCEINLINE void FChaosUserData::Set(void* UserData, UPrimitiveComponent* Payload)		{ check(UserData); ((FChaosUserData*)UserData)->Type = EChaosUserDataType::PrimitiveComponent; ((FChaosUserData*)UserData)->Payload = Payload; }
+template <> FORCEINLINE void FChaosUserData::Set(void* UserData, FKShapeElem* Payload)	{ check(UserData); ((FChaosUserData*)UserData)->Type = EChaosUserDataType::AggShape; ((FChaosUserData*)UserData)->Payload = Payload; }
+template <> FORCEINLINE void FChaosUserData::Set(void* UserData, FCustomChaosPayload* Payload) { check(UserData); ((FChaosUserData*)UserData)->Type = EChaosUserDataType::CustomPayload; ((FChaosUserData*)UserData)->Payload = Payload; }
+
+struct FChaosFilterData
+{
+	FORCEINLINE FChaosFilterData()
 	{
 		word0 = word1 = word2 = word3 = 0;
 	}
 
-	/**
-	\brief Constructor to set filter data initially.
-	*/
-	PX_INLINE PxFilterData(PxU32 w0, PxU32 w1, PxU32 w2, PxU32 w3) : word0(w0), word1(w1), word2(w2), word3(w3) {}
+	FORCEINLINE FChaosFilterData(uint32 w0, uint32 w1, uint32 w2, uint32 w3) : word0(w0), word1(w1), word2(w2), word3(w3) {}
 
-	/**
-	\brief (re)sets the structure to the default.
-	*/
-	PX_INLINE void setToDefault()
+	FORCEINLINE void Reset()
 	{
-		*this = PxFilterData();
+		*this = FChaosFilterData();
 	}
 
-	/**
-	\brief Comparison operator to allow use in Array.
-	*/
-	PX_INLINE bool operator == (const PxFilterData& a) const
+	FORCEINLINE bool operator== (const FChaosFilterData& Other) const
 	{
-		return a.word0 == word0 && a.word1 == word1 && a.word2 == word2 && a.word3 == word3;
+		return Other.word0 == word0 && Other.word1 == word1 && Other.word2 == word2 && Other.word3 == word3;
 	}
 
-	/**
-	\brief Comparison operator to allow use in Array.
-	*/
-	PX_INLINE bool operator != (const PxFilterData& a) const
+	FORCEINLINE bool operator != (const FChaosFilterData& Other) const
 	{
-		return !(a == *this);
+		return !(Other == *this);
 	}
 
-	PxU32 word0;
-	PxU32 word1;
-	PxU32 word2;
-	PxU32 word3;
+	uint32 word0;
+	uint32 word1;
+	uint32 word2;
+	uint32 word3;
 };
 
-struct PxQueryFlag
+struct FChaosQueryFlag
 {
 	enum Enum
 	{
@@ -300,9 +261,9 @@ struct PxQueryFlag
 
 		eDYNAMIC = (1 << 1),	//!< Traverse dynamic shapes
 
-		ePREFILTER = (1 << 2),	//!< Run the pre-intersection-test filter (see #PxQueryFilterCallback::preFilter())
+		ePREFILTER = (1 << 2),	//!< Run the pre-intersection-test filter
 
-		ePOSTFILTER = (1 << 3),	//!< Run the post-intersection-test filter (see #PxQueryFilterCallback::postFilter())
+		ePOSTFILTER = (1 << 3),	//!< Run the post-intersection-test filter
 
 		eANY_HIT = (1 << 4),	//!< Abort traversal as soon as any hit is found and return it via callback.block.
 								//!< Helps query performance. Both eTOUCH and eBLOCK hitTypes are considered hits with this flag.
@@ -314,64 +275,56 @@ struct PxQueryFlag
 	};
 };
 
-typedef PxU8 PxClientID;
-static const PxClientID PX_DEFAULT_CLIENT = 0;
-
-#define PX_CUDA_CALLABLE
-
 template <typename enumtype, typename storagetype = uint32_t>
-class PxFlags
+class ChaosFlags
 {
 public:
 	typedef storagetype InternalType;
 
-	PX_CUDA_CALLABLE PX_INLINE explicit PxFlags(const PxEMPTY)
-	{
-	}
-	PX_CUDA_CALLABLE PX_INLINE PxFlags(void);
-	PX_CUDA_CALLABLE PX_INLINE PxFlags(enumtype e);
-	PX_CUDA_CALLABLE PX_INLINE PxFlags(const PxFlags<enumtype, storagetype>& f);
-	PX_CUDA_CALLABLE PX_INLINE explicit PxFlags(storagetype b);
-
-	PX_CUDA_CALLABLE PX_INLINE bool isSet(enumtype e) const;
-	PX_CUDA_CALLABLE PX_INLINE PxFlags<enumtype, storagetype>& set(enumtype e);
-	PX_CUDA_CALLABLE PX_INLINE bool operator==(enumtype e) const;
-	PX_CUDA_CALLABLE PX_INLINE bool operator==(const PxFlags<enumtype, storagetype>& f) const;
-	PX_CUDA_CALLABLE PX_INLINE bool operator==(bool b) const;
-	PX_CUDA_CALLABLE PX_INLINE bool operator!=(enumtype e) const;
-	PX_CUDA_CALLABLE PX_INLINE bool operator!=(const PxFlags<enumtype, storagetype>& f) const;
-
-	PX_CUDA_CALLABLE PX_INLINE PxFlags<enumtype, storagetype>& operator=(const PxFlags<enumtype, storagetype>& f);
-	PX_CUDA_CALLABLE PX_INLINE PxFlags<enumtype, storagetype>& operator=(enumtype e);
-
-	PX_CUDA_CALLABLE PX_INLINE PxFlags<enumtype, storagetype>& operator|=(enumtype e);
-	PX_CUDA_CALLABLE PX_INLINE PxFlags<enumtype, storagetype>& operator|=(const PxFlags<enumtype, storagetype>& f);
-	PX_CUDA_CALLABLE PX_INLINE PxFlags<enumtype, storagetype> operator|(enumtype e) const;
-	PX_CUDA_CALLABLE PX_INLINE PxFlags<enumtype, storagetype> operator|(const PxFlags<enumtype, storagetype>& f) const;
-
-	PX_CUDA_CALLABLE PX_INLINE PxFlags<enumtype, storagetype>& operator&=(enumtype e);
-	PX_CUDA_CALLABLE PX_INLINE PxFlags<enumtype, storagetype>& operator&=(const PxFlags<enumtype, storagetype>& f);
-	PX_CUDA_CALLABLE PX_INLINE PxFlags<enumtype, storagetype> operator&(enumtype e) const;
-	PX_CUDA_CALLABLE PX_INLINE PxFlags<enumtype, storagetype> operator&(const PxFlags<enumtype, storagetype>& f) const;
-
-	PX_CUDA_CALLABLE PX_INLINE PxFlags<enumtype, storagetype>& operator^=(enumtype e);
-	PX_CUDA_CALLABLE PX_INLINE PxFlags<enumtype, storagetype>& operator^=(const PxFlags<enumtype, storagetype>& f);
-	PX_CUDA_CALLABLE PX_INLINE PxFlags<enumtype, storagetype> operator^(enumtype e) const;
-	PX_CUDA_CALLABLE PX_INLINE PxFlags<enumtype, storagetype> operator^(const PxFlags<enumtype, storagetype>& f) const;
-
-	PX_CUDA_CALLABLE PX_INLINE PxFlags<enumtype, storagetype> operator~(void) const;
-
-	PX_CUDA_CALLABLE PX_INLINE operator bool(void) const;
-	PX_CUDA_CALLABLE PX_INLINE operator uint8_t(void) const;
-	PX_CUDA_CALLABLE PX_INLINE operator uint16_t(void) const;
-	PX_CUDA_CALLABLE PX_INLINE operator uint32_t(void) const;
-
-	PX_CUDA_CALLABLE PX_INLINE void clear(enumtype e);
+	 FORCEINLINE ChaosFlags(void);
+	 FORCEINLINE ChaosFlags(enumtype e);
+	 FORCEINLINE ChaosFlags(const ChaosFlags<enumtype, storagetype>& f);
+	 FORCEINLINE explicit ChaosFlags(storagetype b);
+
+	 FORCEINLINE bool isSet(enumtype e) const;
+	 FORCEINLINE ChaosFlags<enumtype, storagetype>& set(enumtype e);
+	 FORCEINLINE bool operator==(enumtype e) const;
+	 FORCEINLINE bool operator==(const ChaosFlags<enumtype, storagetype>& f) const;
+	 FORCEINLINE bool operator==(bool b) const;
+	 FORCEINLINE bool operator!=(enumtype e) const;
+	 FORCEINLINE bool operator!=(const ChaosFlags<enumtype, storagetype>& f) const;
+
+	 FORCEINLINE ChaosFlags<enumtype, storagetype>& operator=(const ChaosFlags<enumtype, storagetype>& f);
+	 FORCEINLINE ChaosFlags<enumtype, storagetype>& operator=(enumtype e);
+
+	 FORCEINLINE ChaosFlags<enumtype, storagetype>& operator|=(enumtype e);
+	 FORCEINLINE ChaosFlags<enumtype, storagetype>& operator|=(const ChaosFlags<enumtype, storagetype>& f);
+	 FORCEINLINE ChaosFlags<enumtype, storagetype> operator|(enumtype e) const;
+	 FORCEINLINE ChaosFlags<enumtype, storagetype> operator|(const ChaosFlags<enumtype, storagetype>& f) const;
+
+	 FORCEINLINE ChaosFlags<enumtype, storagetype>& operator&=(enumtype e);
+	 FORCEINLINE ChaosFlags<enumtype, storagetype>& operator&=(const ChaosFlags<enumtype, storagetype>& f);
+	 FORCEINLINE ChaosFlags<enumtype, storagetype> operator&(enumtype e) const;
+	 FORCEINLINE ChaosFlags<enumtype, storagetype> operator&(const ChaosFlags<enumtype, storagetype>& f) const;
+
+	 FORCEINLINE ChaosFlags<enumtype, storagetype>& operator^=(enumtype e);
+	 FORCEINLINE ChaosFlags<enumtype, storagetype>& operator^=(const ChaosFlags<enumtype, storagetype>& f);
+	 FORCEINLINE ChaosFlags<enumtype, storagetype> operator^(enumtype e) const;
+	 FORCEINLINE ChaosFlags<enumtype, storagetype> operator^(const ChaosFlags<enumtype, storagetype>& f) const;
+
+	 FORCEINLINE ChaosFlags<enumtype, storagetype> operator~(void) const;
+
+	 FORCEINLINE operator bool(void) const;
+	 FORCEINLINE operator uint8_t(void) const;
+	 FORCEINLINE operator uint16_t(void) const;
+	 FORCEINLINE operator uint32_t(void) const;
+
+	 FORCEINLINE void clear(enumtype e);
 
 public:
-	friend PX_INLINE PxFlags<enumtype, storagetype> operator&(enumtype a, PxFlags<enumtype, storagetype>& b)
+	friend FORCEINLINE ChaosFlags<enumtype, storagetype> operator&(enumtype a, ChaosFlags<enumtype, storagetype>& b)
 	{
-		PxFlags<enumtype, storagetype> out;
+		ChaosFlags<enumtype, storagetype> out;
 		out.mBits = a & b.mBits;
 		return out;
 	}
@@ -380,280 +333,253 @@ private:
 	storagetype mBits;
 };
 
-typedef PxFlags<PxQueryFlag::Enum, PxU16> PxQueryFlags;
+typedef ChaosFlags<FChaosQueryFlag::Enum, uint16> FChaosQueryFlags;
 
-inline PxQueryFlags U2PQueryFlags(FQueryFlags Flags)
+inline FChaosQueryFlags U2CQueryFlags(FQueryFlags Flags)
 {
 	uint32 Result = 0;
 	if (Flags & EQueryFlags::PreFilter)
 	{
-		Result |= PxQueryFlag::ePREFILTER;
+		Result |= FChaosQueryFlag::ePREFILTER;
 	}
 
 	if (Flags & EQueryFlags::PostFilter)
 	{
-		Result |= PxQueryFlag::ePOSTFILTER;
+		Result |= FChaosQueryFlag::ePOSTFILTER;
 	}
 
 	if (Flags & EQueryFlags::AnyHit)
 	{
-		Result |= PxQueryFlag::eANY_HIT;
+		Result |= FChaosQueryFlag::eANY_HIT;
 	}
 
-	return (PxQueryFlags)Result;
+	return (FChaosQueryFlags)Result;
 }
 
-struct PxQueryFilterData
+struct FChaosQueryFilterData
 {
-	/** \brief default constructor */
-	explicit PX_INLINE PxQueryFilterData() : flags(PxQueryFlag::eDYNAMIC | PxQueryFlag::eSTATIC), clientId(PX_DEFAULT_CLIENT) {}
+	explicit FORCEINLINE FChaosQueryFilterData() : flags(FChaosQueryFlag::eDYNAMIC | FChaosQueryFlag::eSTATIC), clientId(0) {}
 
-	/** \brief constructor to set both filter data and filter flags */
-	explicit PX_INLINE PxQueryFilterData(const PxFilterData& fd, PxQueryFlags f) : data(fd), flags(f), clientId(PX_DEFAULT_CLIENT) {}
+	explicit FORCEINLINE FChaosQueryFilterData(const FChaosFilterData& fd, FChaosQueryFlags f) : data(fd), flags(f), clientId(0) {}
 
-	/** \brief constructor to set filter flags only */
-	explicit PX_INLINE PxQueryFilterData(PxQueryFlags f) : flags(f), clientId(PX_DEFAULT_CLIENT) {}
+	explicit FORCEINLINE FChaosQueryFilterData(FChaosQueryFlags f) : flags(f), clientId(0) {}
 
-	PxFilterData	data;		//!< Filter data associated with the scene query
-	PxQueryFlags	flags;		//!< Filter flags (see #PxQueryFlags)
-	PxClientID		clientId;	//!< ID of the client doing the query (see #PxScene.createClient())
+	FChaosFilterData data;
+	FChaosQueryFlags	flags;
+	uint8 clientId;
 };
 
-#define PX_FLAGS_OPERATORS(enumtype, storagetype)                                                                      \
-	PX_INLINE PxFlags<enumtype, storagetype> operator|(enumtype a, enumtype b)                                         \
-	{                                                                                                                  \
-		PxFlags<enumtype, storagetype> r(a);                                                                           \
-		r |= b;                                                                                                        \
-		return r;                                                                                                      \
-	}                                                                                                                  \
-	PX_INLINE PxFlags<enumtype, storagetype> operator&(enumtype a, enumtype b)                                         \
-	{                                                                                                                  \
-		PxFlags<enumtype, storagetype> r(a);                                                                           \
-		r &= b;                                                                                                        \
-		return r;                                                                                                      \
-	}                                                                                                                  \
-	PX_INLINE PxFlags<enumtype, storagetype> operator~(enumtype a)                                                     \
-	{                                                                                                                  \
-		return ~PxFlags<enumtype, storagetype>(a);                                                                     \
-	}
-
-#define PX_FLAGS_TYPEDEF(x, y)                                                                                         \
-	typedef PxFlags<x::Enum, y> x##s;                                                                                  \
-	PX_FLAGS_OPERATORS(x::Enum, y)
-
 template <typename enumtype, typename storagetype>
-PX_INLINE PxFlags<enumtype, storagetype>::PxFlags(void)
+FORCEINLINE ChaosFlags<enumtype, storagetype>::ChaosFlags(void)
 {
 	mBits = 0;
 }
 
 template <typename enumtype, typename storagetype>
-PX_INLINE PxFlags<enumtype, storagetype>::PxFlags(enumtype e)
+FORCEINLINE ChaosFlags<enumtype, storagetype>::ChaosFlags(enumtype e)
 {
 	mBits = static_cast<storagetype>(e);
 }
 
 template <typename enumtype, typename storagetype>
-PX_INLINE PxFlags<enumtype, storagetype>::PxFlags(const PxFlags<enumtype, storagetype>& f)
+FORCEINLINE ChaosFlags<enumtype, storagetype>::ChaosFlags(const ChaosFlags<enumtype, storagetype>& f)
 {
 	mBits = f.mBits;
 }
 
 template <typename enumtype, typename storagetype>
-PX_INLINE PxFlags<enumtype, storagetype>::PxFlags(storagetype b)
+FORCEINLINE ChaosFlags<enumtype, storagetype>::ChaosFlags(storagetype b)
 {
 	mBits = b;
 }
 
 template <typename enumtype, typename storagetype>
-PX_INLINE bool PxFlags<enumtype, storagetype>::isSet(enumtype e) const
+FORCEINLINE bool ChaosFlags<enumtype, storagetype>::isSet(enumtype e) const
 {
 	return (mBits & static_cast<storagetype>(e)) == static_cast<storagetype>(e);
 }
 
 template <typename enumtype, typename storagetype>
-PX_INLINE PxFlags<enumtype, storagetype>& PxFlags<enumtype, storagetype>::set(enumtype e)
+FORCEINLINE ChaosFlags<enumtype, storagetype>& ChaosFlags<enumtype, storagetype>::set(enumtype e)
 {
 	mBits = static_cast<storagetype>(e);
 	return *this;
 }
 
 template <typename enumtype, typename storagetype>
-PX_INLINE bool PxFlags<enumtype, storagetype>::operator==(enumtype e) const
+FORCEINLINE bool ChaosFlags<enumtype, storagetype>::operator==(enumtype e) const
 {
 	return mBits == static_cast<storagetype>(e);
 }
 
 template <typename enumtype, typename storagetype>
-PX_INLINE bool PxFlags<enumtype, storagetype>::operator==(const PxFlags<enumtype, storagetype>& f) const
+FORCEINLINE bool ChaosFlags<enumtype, storagetype>::operator==(const ChaosFlags<enumtype, storagetype>& f) const
 {
 	return mBits == f.mBits;
 }
 
 template <typename enumtype, typename storagetype>
-PX_INLINE bool PxFlags<enumtype, storagetype>::operator==(bool b) const
+FORCEINLINE bool ChaosFlags<enumtype, storagetype>::operator==(bool b) const
 {
 	return bool(*this) == b;
 }
 
 template <typename enumtype, typename storagetype>
-PX_INLINE bool PxFlags<enumtype, storagetype>::operator!=(enumtype e) const
+FORCEINLINE bool ChaosFlags<enumtype, storagetype>::operator!=(enumtype e) const
 {
 	return mBits != static_cast<storagetype>(e);
 }
 
 template <typename enumtype, typename storagetype>
-PX_INLINE bool PxFlags<enumtype, storagetype>::operator!=(const PxFlags<enumtype, storagetype>& f) const
+FORCEINLINE bool ChaosFlags<enumtype, storagetype>::operator!=(const ChaosFlags<enumtype, storagetype>& f) const
 {
 	return mBits != f.mBits;
 }
 
 template <typename enumtype, typename storagetype>
-PX_INLINE PxFlags<enumtype, storagetype>& PxFlags<enumtype, storagetype>::operator=(enumtype e)
+FORCEINLINE ChaosFlags<enumtype, storagetype>& ChaosFlags<enumtype, storagetype>::operator=(enumtype e)
 {
 	mBits = static_cast<storagetype>(e);
 	return *this;
 }
 
 template <typename enumtype, typename storagetype>
-PX_INLINE PxFlags<enumtype, storagetype>& PxFlags<enumtype, storagetype>::operator=(const PxFlags<enumtype, storagetype>& f)
+FORCEINLINE ChaosFlags<enumtype, storagetype>& ChaosFlags<enumtype, storagetype>::operator=(const ChaosFlags<enumtype, storagetype>& f)
 {
 	mBits = f.mBits;
 	return *this;
 }
 
 template <typename enumtype, typename storagetype>
-PX_INLINE PxFlags<enumtype, storagetype>& PxFlags<enumtype, storagetype>::operator|=(enumtype e)
+FORCEINLINE ChaosFlags<enumtype, storagetype>& ChaosFlags<enumtype, storagetype>::operator|=(enumtype e)
 {
 	mBits |= static_cast<storagetype>(e);
 	return *this;
 }
 
 template <typename enumtype, typename storagetype>
-PX_INLINE PxFlags<enumtype, storagetype>& PxFlags<enumtype, storagetype>::
-operator|=(const PxFlags<enumtype, storagetype>& f)
+FORCEINLINE ChaosFlags<enumtype, storagetype>& ChaosFlags<enumtype, storagetype>::
+operator|=(const ChaosFlags<enumtype, storagetype>& f)
 {
 	mBits |= f.mBits;
 	return *this;
 }
 
 template <typename enumtype, typename storagetype>
-PX_INLINE PxFlags<enumtype, storagetype> PxFlags<enumtype, storagetype>::operator|(enumtype e) const
+FORCEINLINE ChaosFlags<enumtype, storagetype> ChaosFlags<enumtype, storagetype>::operator|(enumtype e) const
 {
-	PxFlags<enumtype, storagetype> out(*this);
+	ChaosFlags<enumtype, storagetype> out(*this);
 	out |= e;
 	return out;
 }
 
 template <typename enumtype, typename storagetype>
-PX_INLINE PxFlags<enumtype, storagetype> PxFlags<enumtype, storagetype>::
-operator|(const PxFlags<enumtype, storagetype>& f) const
+FORCEINLINE ChaosFlags<enumtype, storagetype> ChaosFlags<enumtype, storagetype>::
+operator|(const ChaosFlags<enumtype, storagetype>& f) const
 {
-	PxFlags<enumtype, storagetype> out(*this);
+	ChaosFlags<enumtype, storagetype> out(*this);
 	out |= f;
 	return out;
 }
 
 template <typename enumtype, typename storagetype>
-PX_INLINE PxFlags<enumtype, storagetype>& PxFlags<enumtype, storagetype>::operator&=(enumtype e)
+FORCEINLINE ChaosFlags<enumtype, storagetype>& ChaosFlags<enumtype, storagetype>::operator&=(enumtype e)
 {
 	mBits &= static_cast<storagetype>(e);
 	return *this;
 }
 
 template <typename enumtype, typename storagetype>
-PX_INLINE PxFlags<enumtype, storagetype>& PxFlags<enumtype, storagetype>::
-operator&=(const PxFlags<enumtype, storagetype>& f)
+FORCEINLINE ChaosFlags<enumtype, storagetype>& ChaosFlags<enumtype, storagetype>::
+operator&=(const ChaosFlags<enumtype, storagetype>& f)
 {
 	mBits &= f.mBits;
 	return *this;
 }
 
 template <typename enumtype, typename storagetype>
-PX_INLINE PxFlags<enumtype, storagetype> PxFlags<enumtype, storagetype>::operator&(enumtype e) const
+FORCEINLINE ChaosFlags<enumtype, storagetype> ChaosFlags<enumtype, storagetype>::operator&(enumtype e) const
 {
-	PxFlags<enumtype, storagetype> out = *this;
+	ChaosFlags<enumtype, storagetype> out = *this;
 	out.mBits &= static_cast<storagetype>(e);
 	return out;
 }
 
 template <typename enumtype, typename storagetype>
-PX_INLINE PxFlags<enumtype, storagetype> PxFlags<enumtype, storagetype>::
-operator&(const PxFlags<enumtype, storagetype>& f) const
+FORCEINLINE ChaosFlags<enumtype, storagetype> ChaosFlags<enumtype, storagetype>::
+operator&(const ChaosFlags<enumtype, storagetype>& f) const
 {
-	PxFlags<enumtype, storagetype> out = *this;
+	ChaosFlags<enumtype, storagetype> out = *this;
 	out.mBits &= f.mBits;
 	return out;
 }
 
 template <typename enumtype, typename storagetype>
-PX_INLINE PxFlags<enumtype, storagetype>& PxFlags<enumtype, storagetype>::operator^=(enumtype e)
+FORCEINLINE ChaosFlags<enumtype, storagetype>& ChaosFlags<enumtype, storagetype>::operator^=(enumtype e)
 {
 	mBits ^= static_cast<storagetype>(e);
 	return *this;
 }
 
 template <typename enumtype, typename storagetype>
-PX_INLINE PxFlags<enumtype, storagetype>& PxFlags<enumtype, storagetype>::
-operator^=(const PxFlags<enumtype, storagetype>& f)
+FORCEINLINE ChaosFlags<enumtype, storagetype>& ChaosFlags<enumtype, storagetype>::
+operator^=(const ChaosFlags<enumtype, storagetype>& f)
 {
 	mBits ^= f.mBits;
 	return *this;
 }
 
 template <typename enumtype, typename storagetype>
-PX_INLINE PxFlags<enumtype, storagetype> PxFlags<enumtype, storagetype>::operator^(enumtype e) const
+FORCEINLINE ChaosFlags<enumtype, storagetype> ChaosFlags<enumtype, storagetype>::operator^(enumtype e) const
 {
-	PxFlags<enumtype, storagetype> out = *this;
+	ChaosFlags<enumtype, storagetype> out = *this;
 	out.mBits ^= static_cast<storagetype>(e);
 	return out;
 }
 
 template <typename enumtype, typename storagetype>
-PX_INLINE PxFlags<enumtype, storagetype> PxFlags<enumtype, storagetype>::
-operator^(const PxFlags<enumtype, storagetype>& f) const
+FORCEINLINE ChaosFlags<enumtype, storagetype> ChaosFlags<enumtype, storagetype>::
+operator^(const ChaosFlags<enumtype, storagetype>& f) const
 {
-	PxFlags<enumtype, storagetype> out = *this;
+	ChaosFlags<enumtype, storagetype> out = *this;
 	out.mBits ^= f.mBits;
 	return out;
 }
 
 template <typename enumtype, typename storagetype>
-PX_INLINE PxFlags<enumtype, storagetype> PxFlags<enumtype, storagetype>::operator~(void) const
+FORCEINLINE ChaosFlags<enumtype, storagetype> ChaosFlags<enumtype, storagetype>::operator~(void) const
 {
-	PxFlags<enumtype, storagetype> out;
+	ChaosFlags<enumtype, storagetype> out;
 	out.mBits = storagetype(~mBits);
 	return out;
 }
 
 template <typename enumtype, typename storagetype>
-PX_INLINE PxFlags<enumtype, storagetype>::operator bool(void) const
+FORCEINLINE ChaosFlags<enumtype, storagetype>::operator bool(void) const
 {
 	return mBits ? true : false;
 }
 
 template <typename enumtype, typename storagetype>
-PX_INLINE PxFlags<enumtype, storagetype>::operator uint8_t(void) const
+FORCEINLINE ChaosFlags<enumtype, storagetype>::operator uint8_t(void) const
 {
 	return static_cast<uint8_t>(mBits);
 }
 
 template <typename enumtype, typename storagetype>
-PX_INLINE PxFlags<enumtype, storagetype>::operator uint16_t(void) const
+FORCEINLINE ChaosFlags<enumtype, storagetype>::operator uint16_t(void) const
 {
 	return static_cast<uint16_t>(mBits);
 }
 
 template <typename enumtype, typename storagetype>
-PX_INLINE PxFlags<enumtype, storagetype>::operator uint32_t(void) const
+FORCEINLINE ChaosFlags<enumtype, storagetype>::operator uint32_t(void) const
 {
 	return static_cast<uint32_t>(mBits);
 }
 
 template <typename enumtype, typename storagetype>
-PX_INLINE void PxFlags<enumtype, storagetype>::clear(enumtype e)
+FORCEINLINE void ChaosFlags<enumtype, storagetype>::clear(enumtype e)
 {
 	mBits &= ~static_cast<storagetype>(e);
 }
-
-#endif // !WITH_PHYSX
diff --git a/Engine/Source/Runtime/PhysicsCore/Public/PhysicsInterfaceUtilsCore.h b/Engine/Source/Runtime/PhysicsCore/Public/PhysicsInterfaceUtilsCore.h
--- a/Engine/Source/Runtime/PhysicsCore/Public/PhysicsInterfaceUtilsCore.h
+++ b/Engine/Source/Runtime/PhysicsCore/Public/PhysicsInterfaceUtilsCore.h
@@ -11,11 +11,13 @@
 #include "CollisionShape.h"
 #endif
 
-PHYSICSCORE_API FCollisionFilterData P2UFilterData(const PxFilterData& PFilterData);
-PHYSICSCORE_API PxFilterData U2PFilterData(const FCollisionFilterData& FilterData);
+PHYSICSCORE_API FCollisionFilterData C2UFilterData(const FChaosFilterData& FilterData);
+PHYSICSCORE_API FChaosFilterData U2CFilterData(const FCollisionFilterData& FilterData);
 
 #if PHYSICS_INTERFACE_PHYSX
 
+PHYSICSCORE_API FCollisionFilterData P2UFilterData(const PxFilterData& PFilterData);
+PHYSICSCORE_API PxFilterData U2PFilterData(const FCollisionFilterData& FilterData);
 PHYSICSCORE_API PxShapeFlags BuildPhysXShapeFlags(FBodyCollisionFlags BodyCollisionFlags, bool bPhysicsStatic, bool bIsTriangleMesh);
 PHYSICSCORE_API PxGeometryType::Enum U2PCollisionShapeType(ECollisionShapeType InUType);
 PHYSICSCORE_API ECollisionShapeType P2UCollisionShapeType(PxGeometryType::Enum InPType);
diff --git a/Engine/Source/Runtime/PhysicsSQ/Private/SQAccelerator.cpp b/Engine/Source/Runtime/PhysicsSQ/Private/SQAccelerator.cpp
--- a/Engine/Source/Runtime/PhysicsSQ/Private/SQAccelerator.cpp
+++ b/Engine/Source/Runtime/PhysicsSQ/Private/SQAccelerator.cpp
@@ -95,12 +95,17 @@ struct TSQVisitor : public Chaos::ISpatialVisitor<TPayload, float>
 		, DebugParams(InDebugParams)
 		, HitBuffer(InHitBuffer)
 		, QueryFilterData(InQueryFilterData)
+#if PHYSICS_INTERFACE_PHYSX
 		, QueryFilterDataConcrete(P2UFilterData(QueryFilterData.data))
+#else
+		, QueryFilterDataConcrete(C2UFilterData(QueryFilterData.data))
+#endif
 		, QueryCallback(InQueryCallback)
 	{
-#if WITH_PHYSX
-		//#TODO - reimplement query flags alternative for Chaos
+#if PHYSICS_INTERFACE_PHYSX
 		bAnyHit = QueryFilterData.flags & PxQueryFlag::eANY_HIT;
+#else
+		bAnyHit = QueryFilterData.flags & FChaosQueryFlag::eANY_HIT;
 #endif
 	}
 
@@ -113,13 +118,20 @@ struct TSQVisitor : public Chaos::ISpatialVisitor<TPayload, float>
 		, DebugParams(InDebugParams)
 		, HitBuffer(InHitBuffer)
 		, QueryFilterData(InQueryFilterData)
+#if PHYSICS_INTERFACE_PHYSX
 		, QueryFilterDataConcrete(P2UFilterData(QueryFilterData.data))
+#else
+		, QueryFilterDataConcrete(C2UFilterData(QueryFilterData.data))
+#endif
 		, QueryGeom(&InQueryGeom)
 		, QueryCallback(InQueryCallback)
 		, StartTM(InStartTM)
 	{
-		//#TODO - reimplement query flags alternative for Chaos
+#if PHYSICS_INTERFACE_PHYSX
 		bAnyHit = QueryFilterData.flags & PxQueryFlag::eANY_HIT;
+#else
+		bAnyHit = QueryFilterData.flags & FChaosQueryFlag::eANY_HIT;
+#endif
 
 		//todo: check THitType is sweep
 	}
@@ -131,13 +143,20 @@ struct TSQVisitor : public Chaos::ISpatialVisitor<TPayload, float>
 		, DebugParams(InDebugParams)
 		, HitBuffer(InHitBuffer)
 		, QueryFilterData(InQueryFilterData)
+#if PHYSICS_INTERFACE_PHYSX
 		, QueryFilterDataConcrete(P2UFilterData(QueryFilterData.data))
+#else
+		, QueryFilterDataConcrete(C2UFilterData(QueryFilterData.data))
+#endif
 		, QueryGeom(&InQueryGeom)
 		, QueryCallback(InQueryCallback)
 		, StartTM(InWorldTM)
 	{
-		//#TODO - reimplement query flags alternative for Chaos
+#if PHYSICS_INTERFACE_PHYSX
 		bAnyHit = QueryFilterData.flags & PxQueryFlag::eANY_HIT;
+#else
+		bAnyHit = QueryFilterData.flags & FChaosQueryFlag::eANY_HIT;
+#endif
 
 		//todo: check THitType is overlap
 	}
@@ -231,8 +250,12 @@ private:
 			}
 
 			//TODO: use gt particles directly
-			//#TODO alternative to px flags
+#if PHYSICS_INTERFACE_PHYSX
 			ECollisionQueryHitType HitType = QueryFilterData.flags & PxQueryFlag::ePREFILTER ? QueryCallback.PreFilter(QueryFilterDataConcrete, *Shape, *GeometryParticle) : ECollisionQueryHitType::Block;
+#else
+			ECollisionQueryHitType HitType = QueryFilterData.flags & FChaosQueryFlag::ePREFILTER ? QueryCallback.PreFilter(QueryFilterDataConcrete, *Shape, *GeometryParticle) : ECollisionQueryHitType::Block;
+#endif
+
 
 			if (HitType != ECollisionQueryHitType::None)
 			{
@@ -288,7 +311,12 @@ private:
 					//QUICK_SCOPE_CYCLE_COUNTER(SQNarrowHit);
 					FillHitHelper(Hit, Distance, WorldPosition, WorldNormal, FaceIdx, bComputeMTD);
 
+
+#if PHYSICS_INTERFACE_PHYSX
 					HitType = QueryFilterData.flags & PxQueryFlag::ePOSTFILTER ? QueryCallback.PostFilter(QueryFilterDataConcrete, Hit) : HitType;
+#else
+					HitType = QueryFilterData.flags & FChaosQueryFlag::ePOSTFILTER ? QueryCallback.PostFilter(QueryFilterDataConcrete, Hit) : HitType;
+#endif
 
 					if (HitType != ECollisionQueryHitType::None)
 					{
