Chaos : Collision Manager
- Temporarily disable the IgnoreCollisionsManager

diff --git a/Engine/Source/Runtime/PhysicsCore/Private/ChaosEngineInterface.cpp b/Engine/Source/Runtime/PhysicsCore/Private/ChaosEngineInterface.cpp
--- a/Engine/Source/Runtime/PhysicsCore/Private/ChaosEngineInterface.cpp
+++ b/Engine/Source/Runtime/PhysicsCore/Private/ChaosEngineInterface.cpp
@@ -27,7 +27,7 @@ FPhysicsDelegatesCore::FOnUpdatePhysXMaterial FPhysicsDelegatesCore::OnUpdatePhy
 bool bEnableChaosJointConstraints = true;
 FAutoConsoleVariableRef CVarEnableChaosJointConstraints(TEXT("p.ChaosSolverEnableJointConstraints"), bEnableChaosJointConstraints, TEXT("Enable Joint Constraints defined within the Physics Asset Editor"));
 
-bool bEnableChaosCollisionManager = true;
+bool bEnableChaosCollisionManager = false;
 FAutoConsoleVariableRef CVarEnableChaosCollisionManager(TEXT("p.Chaos.Collision.EnableCollisionManager"), bEnableChaosCollisionManager, TEXT("Enable Chaos's Collision Manager for ignoring collisions between rigid bodies. [def:1]"));
 
 bool FPhysicsConstraintReference_Chaos::IsValid() const
