Chaos : Ignore Collisions
- Reenabled the collision manager.
- Added threading handoff for the collision constraint flag.
- Removed game thread particle dependency in the ignore collisions update.

diff --git a/Engine/Source/Runtime/Core/Public/UObject/UE5MainStreamObjectVersion.h b/Engine/Source/Runtime/Core/Public/UObject/UE5MainStreamObjectVersion.h
--- a/Engine/Source/Runtime/Core/Public/UObject/UE5MainStreamObjectVersion.h
+++ b/Engine/Source/Runtime/Core/Public/UObject/UE5MainStreamObjectVersion.h
@@ -58,6 +58,9 @@ struct CORE_API FUE5MainStreamObjectVersion
 		// Added FLandscapeSplineActorDesc
 		AddedLandscapeSplineActorDesc,
 
+		// Added support for per-object collision constraint flag. [Chaos]
+		AddCollisionConstraintFlag,
+
 		// -----<new versions can be added above this line>-------------------------------------------------
 		VersionPlusOne,
 		LatestVersion = VersionPlusOne - 1
diff --git a/Engine/Source/Runtime/Experimental/Chaos/Private/Chaos/Collision/CollisionConstraintFlags.cpp b/Engine/Source/Runtime/Experimental/Chaos/Private/Chaos/Collision/CollisionConstraintFlags.cpp
--- a/Engine/Source/Runtime/Experimental/Chaos/Private/Chaos/Collision/CollisionConstraintFlags.cpp
+++ b/Engine/Source/Runtime/Experimental/Chaos/Private/Chaos/Collision/CollisionConstraintFlags.cpp
@@ -90,13 +90,13 @@ namespace Chaos
 
 		// remove particles that have been created and destroyed
 		// before the queue was ever processed. 
-		TArray<FGeometryParticle*> PreculledParticles;
+		TArray<FHandleID> PreculledParticles;
 		if (PendingActivations.Num() && PendingDeactivations.Num())
 		{
-			TArray<FGeometryParticle*> DeletionList;
+			TArray<FHandleID> DeletionList;
 			for (auto& Elem : PendingActivations)
 			{
-				int32 DeactiveIndex = PendingDeactivations.Find(Elem.Key->UniqueIdx());
+				int32 DeactiveIndex = PendingDeactivations.Find(Elem.Key);
 				if (DeactiveIndex != INDEX_NONE)
 				{
 					DeletionList.Add(Elem.Key);
@@ -104,7 +104,7 @@ namespace Chaos
 					PendingDeactivations.RemoveAtSwap(DeactiveIndex, 1);
 				}
 			}
-			for (auto* Del : DeletionList)
+			for (FHandleID Del : DeletionList)
 				PendingActivations.Remove(Del);
 		}
 
@@ -113,53 +113,32 @@ namespace Chaos
 		// simulation. 
 		if (PendingActivations.Num())
 		{
-			TArray<FGeometryParticle*> DeletionList;
+			TArray<FHandleID> DeletionList;
 			for (auto& Elem : PendingActivations)
 			{
-				if (Elem.Key)
+				for (int Index = Elem.Value.Num() - 1; Index >= 0; Index--)
 				{
-					if (Elem.Key->IsParticleValid())
+					if (PreculledParticles.Contains(Elem.Value[Index]))
 					{
-						if (TGeometryParticleHandle<FReal, 3>* Handle0 = Elem.Key->Handle())
+						Elem.Value.RemoveAtSwap(Index, 1);
+					}
+					else
+					{
+						FUniqueIdx ID0 = Elem.Key;
+						FUniqueIdx ID1 = Elem.Value[Index];
+						if (!IgnoresCollision(ID0, ID1))
 						{
-							for (int Index = Elem.Value.Num() - 1; Index >= 0; Index--)
-							{
-								if (TGeometryParticle<FReal, 3>* Val = Elem.Value[Index])
-								{
-									if (Val->IsParticleValid())
-									{
-										if (PreculledParticles.Contains(Val))
-										{
-											Elem.Value.RemoveAtSwap(Index, 1);
-										}
-										else if (TGeometryParticleHandle<FReal, 3>* Handle1 = Val->Handle())
-										{
-											FUniqueIdx ID0 = Handle0->UniqueIdx();
-											FUniqueIdx ID1 = Handle1->UniqueIdx();
-											if (!IgnoresCollision(ID0, ID1))
-											{
-												Chaos::TPBDRigidParticleHandle<FReal, 3>* ParticleHandle0 = Handle0->CastToRigidParticle();
-												Chaos::TPBDRigidParticleHandle<FReal, 3>* ParticleHandle1 = Handle1->CastToRigidParticle();
-
-												ParticleHandle0->AddCollisionConstraintFlag(Chaos::ECollisionConstraintFlags::CCF_BroadPhaseIgnoreCollisions);
-												AddIgnoreCollisionsFor(ID0, ID1);
-
-												ParticleHandle1->AddCollisionConstraintFlag(Chaos::ECollisionConstraintFlags::CCF_BroadPhaseIgnoreCollisions);
-												AddIgnoreCollisionsFor(ID1, ID0);
-											}
-
-											Elem.Value.RemoveAtSwap(Index, 1);
-										}
-									}
-								}
-							}
+							AddIgnoreCollisionsFor(ID0, ID1);
+							AddIgnoreCollisionsFor(ID1, ID0);
 						}
-						if (!Elem.Value.Num())
-							DeletionList.Add(Elem.Key);
+
+						Elem.Value.RemoveAtSwap(Index, 1);
 					}
 				}
+				if (!Elem.Value.Num())
+					DeletionList.Add(Elem.Key);
 			}
-			for (auto* Del : DeletionList)
+			for (FHandleID Del : DeletionList)
 				PendingActivations.Remove(Del);
 		}
 
diff --git a/Engine/Source/Runtime/Experimental/Chaos/Public/Chaos/Collision/CollisionConstraintFlags.h b/Engine/Source/Runtime/Experimental/Chaos/Public/Chaos/Collision/CollisionConstraintFlags.h
--- a/Engine/Source/Runtime/Experimental/Chaos/Public/Chaos/Collision/CollisionConstraintFlags.h
+++ b/Engine/Source/Runtime/Experimental/Chaos/Public/Chaos/Collision/CollisionConstraintFlags.h
@@ -20,11 +20,10 @@ namespace Chaos
 	class CHAOS_API FIgnoreCollisionManager
 	{
 	public:
-		using FGeometryParticle = TGeometryParticle<FReal, 3>;
 		using FHandleID = FUniqueIdx;
 		using FDeactivationArray = TArray<FUniqueIdx>;
 		using FActiveMap = TMap<FHandleID, TArray<FHandleID> >;
-		using FPendingMap = TMap<FGeometryParticle*, TArray<FGeometryParticle*> >;
+		using FPendingMap = TMap<FHandleID, TArray<FHandleID> >;
 		struct FStorageData
 		{
 			FPendingMap PendingActivations;
diff --git a/Engine/Source/Runtime/Experimental/Chaos/Public/Chaos/ParticleDirtyFlags.h b/Engine/Source/Runtime/Experimental/Chaos/Public/Chaos/ParticleDirtyFlags.h
--- a/Engine/Source/Runtime/Experimental/Chaos/Public/Chaos/ParticleDirtyFlags.h
+++ b/Engine/Source/Runtime/Experimental/Chaos/Public/Chaos/ParticleDirtyFlags.h
@@ -12,6 +12,7 @@
 #include "Chaos/KinematicTargets.h"
 #include "UObject/ExternalPhysicsCustomObjectVersion.h"
 #include "UObject/ExternalPhysicsMaterialCustomObjectVersion.h"
+#include "UObject/UE5MainStreamObjectVersion.h"
 
 class FName;
 
@@ -249,6 +250,7 @@ class FParticleDynamicMisc
 public:
 	void Serialize(FChaosArchive& Ar)
 	{
+		Ar.UsingCustomVersion(FUE5MainStreamObjectVersion::GUID);
 		Ar.UsingCustomVersion(FExternalPhysicsCustomObjectVersion::GUID);
 		Ar << MLinearEtherDrag;
 		Ar << MAngularEtherDrag;
@@ -262,6 +264,12 @@ public:
 		{
 			MOneWayInteraction = false;
 		}
+
+		if (Ar.CustomVer(FUE5MainStreamObjectVersion::GUID) >= FUE5MainStreamObjectVersion::AddCollisionConstraintFlag)
+		{
+			Ar << MCollisionConstraintFlag;
+		}
+
 	}
 
 	template <typename TOther>
@@ -274,6 +282,7 @@ public:
 		SetCollisionGroup(Other.CollisionGroup());
 		SetResimType(Other.ResimType());
 		SetOneWayInteraction(Other.OneWayInteraction());
+		SetCollisionConstraintFlag(Other.CollisionConstraintFlag());
 	}
 
 	template <typename TOther>
@@ -285,7 +294,8 @@ public:
 			&& GravityEnabled() == Other.GravityEnabled()
 			&& CollisionGroup() == Other.CollisionGroup()
 			&& ResimType() == Other.ResimType()
-			&& OneWayInteraction() == Other.OneWayInteraction();
+			&& OneWayInteraction() == Other.OneWayInteraction()
+			&& CollisionConstraintFlag() == Other.CollisionConstraintFlag();
 	}
 
 	bool operator==(const FParticleDynamicMisc& Other) const
@@ -314,6 +324,9 @@ public:
 	bool OneWayInteraction() const { return MOneWayInteraction; }
 	void SetOneWayInteraction(bool InOneWayInteraction) { MOneWayInteraction = InOneWayInteraction; }
 
+	uint32 CollisionConstraintFlag() const { return MCollisionConstraintFlag; }
+	void SetCollisionConstraintFlag(uint32 InCollisionConstraintFlag) { MCollisionConstraintFlag = InCollisionConstraintFlag; }
+
 private:
 	FReal MLinearEtherDrag;
 	FReal MAngularEtherDrag;
@@ -324,6 +337,7 @@ private:
 
 	bool MGravityEnabled;
 	bool MOneWayInteraction = false;
+	uint32 MCollisionConstraintFlag = 0;
 };
 
 inline FChaosArchive& operator<<(FChaosArchive& Ar,FParticleDynamicMisc& Data)
diff --git a/Engine/Source/Runtime/Experimental/Chaos/Public/Chaos/ParticleHandle.h b/Engine/Source/Runtime/Experimental/Chaos/Public/Chaos/ParticleHandle.h
--- a/Engine/Source/Runtime/Experimental/Chaos/Public/Chaos/ParticleHandle.h
+++ b/Engine/Source/Runtime/Experimental/Chaos/Public/Chaos/ParticleHandle.h
@@ -738,6 +738,7 @@ public:
 	void AddCollisionConstraintFlag(const ECollisionConstraintFlags Flag) { PBDRigidParticles->AddCollisionConstraintFlag(Flag, ParticleIdx); }
 	void RemoveCollisionConstraintFlag(const ECollisionConstraintFlags Flag) { PBDRigidParticles->RemoveCollisionConstraintFlag(Flag, ParticleIdx); }
 	void ClearCollisionConstraintFlag() { PBDRigidParticles->ClearCollisionConstraintFlag(ParticleIdx); }
+	uint32 CollisionConstraintFlag() const { return PBDRigidParticles->CollisionConstraintFlag(ParticleIdx); }
 
 	bool Disabled() const { return PBDRigidParticles->Disabled(ParticleIdx); }
 	bool& Disabled() { return PBDRigidParticles->DisabledRef(ParticleIdx); }
@@ -818,6 +819,7 @@ public:
 		SetGravityEnabled(DynamicMisc.GravityEnabled());
 		SetResimType(DynamicMisc.ResimType());
 		SetOneWayInteraction(DynamicMisc.OneWayInteraction());
+		AddCollisionConstraintFlag( (Chaos::ECollisionConstraintFlags)DynamicMisc.CollisionConstraintFlag() );
 	}
 
 	void ResetSmoothedVelocities()
@@ -2252,6 +2254,12 @@ public:
 		MMiscData.Modify(true, MDirtyFlags, Proxy, [InOneWayInteraction](auto& Data) { Data.SetOneWayInteraction(InOneWayInteraction); });
 	}
 
+	uint32 CollisionConstraintFlag() const { return MMiscData.Read().CollisionConstraintFlag(); }
+	void SetCollisionConstraintFlag(const uint32 InCollisionConstraintFlag)
+	{
+		MMiscData.Modify(true, MDirtyFlags, Proxy, [InCollisionConstraintFlag](auto& Data) { Data.SetCollisionConstraintFlag(InCollisionConstraintFlag); });
+	}
+
 	//todo: remove this
 	bool IsInitialized() const { return MInitialized; }
 	void SetInitialized(const bool InInitialized)
diff --git a/Engine/Source/Runtime/Experimental/Chaos/Public/Chaos/RigidParticles.h b/Engine/Source/Runtime/Experimental/Chaos/Public/Chaos/RigidParticles.h
--- a/Engine/Source/Runtime/Experimental/Chaos/Public/Chaos/RigidParticles.h
+++ b/Engine/Source/Runtime/Experimental/Chaos/Public/Chaos/RigidParticles.h
@@ -202,6 +202,7 @@ public:
 	FORCEINLINE void AddCollisionConstraintFlag(const ECollisionConstraintFlags Flag, const int32 Index) { MCollisionConstraintFlags[Index] |= (uint32)Flag; }
 	FORCEINLINE void RemoveCollisionConstraintFlag(const ECollisionConstraintFlags Flag, const int32 Index) { MCollisionConstraintFlags[Index] &= ~(uint32)Flag; }
 	FORCEINLINE void ClearCollisionConstraintFlag(const int32 Index) { MCollisionConstraintFlags[Index] = (uint32)ECollisionConstraintFlags::CCF_None; }
+	FORCEINLINE uint32 CollisionConstraintFlag(const int32 Index) const { return MCollisionConstraintFlags[Index]; }
 
 	FORCEINLINE const bool Disabled(const int32 Index) const { return MDisabled[Index]; }
 
diff --git a/Engine/Source/Runtime/PhysicsCore/Private/ChaosEngineInterface.cpp b/Engine/Source/Runtime/PhysicsCore/Private/ChaosEngineInterface.cpp
--- a/Engine/Source/Runtime/PhysicsCore/Private/ChaosEngineInterface.cpp
+++ b/Engine/Source/Runtime/PhysicsCore/Private/ChaosEngineInterface.cpp
@@ -27,7 +27,7 @@ FPhysicsDelegatesCore::FOnUpdatePhysXMaterial FPhysicsDelegatesCore::OnUpdatePhy
 bool bEnableChaosJointConstraints = true;
 FAutoConsoleVariableRef CVarEnableChaosJointConstraints(TEXT("p.ChaosSolverEnableJointConstraints"), bEnableChaosJointConstraints, TEXT("Enable Joint Constraints defined within the Physics Asset Editor"));
 
-bool bEnableChaosCollisionManager = false;
+bool bEnableChaosCollisionManager = true;
 FAutoConsoleVariableRef CVarEnableChaosCollisionManager(TEXT("p.Chaos.Collision.EnableCollisionManager"), bEnableChaosCollisionManager, TEXT("Enable Chaos's Collision Manager for ignoring collisions between rigid bodies. [def:1]"));
 
 bool FPhysicsConstraintReference_Chaos::IsValid() const
@@ -284,16 +284,35 @@ void FChaosEngineInterface::AddDisabledCollisionsFor_AssumesLocked(const TMap<FP
 		for (auto Elem : InMap)
 		{
 			FPhysicsActorHandle& ActorReference = Elem.Key;
-			TArray< FPhysicsActorHandle >& DisabledCollisions = Elem.Value;
+			Chaos::FUniqueIdx ActorIndex = ActorReference->UniqueIdx();
+
 			Chaos::FPhysicsSolver* Solver = ActorReference->GetProxy()->GetSolver<Chaos::FPhysicsSolver>();
 			Chaos::FIgnoreCollisionManager& CollisionManager = Solver->GetEvolution()->GetBroadPhase().GetIgnoreCollisionManager();
 			int32 ExternalTimestamp = Solver->GetMarshallingManager().GetExternalTimestamp_External();
-			Chaos::FIgnoreCollisionManager::FPendingMap& PendingMap = CollisionManager.GetPendingActivationsForGameThread(ExternalTimestamp);
-			if (PendingMap.Contains(ActorReference))
+			Chaos::FIgnoreCollisionManager::FPendingMap& ActivationMap = CollisionManager.GetPendingActivationsForGameThread(ExternalTimestamp);
+
+			if (ActivationMap.Contains(ActorIndex))
+			{
+				ActivationMap.Remove(ActorIndex);
+			}
+
+			TArray< Chaos::FUniqueIdx > DisabledCollisions;
+			DisabledCollisions.Reserve(Elem.Value.Num());
+
+			if (Chaos::TPBDRigidParticle<float, 3>* Rigid0 = ActorReference->CastToRigidParticle())
 			{
-				PendingMap.Remove(ActorReference);
+				Rigid0->SetCollisionConstraintFlag((uint32)Chaos::ECollisionConstraintFlags::CCF_BroadPhaseIgnoreCollisions);
+				for (auto Handle1 : Elem.Value)
+				{
+					if (Chaos::TPBDRigidParticle<float, 3>* Rigid1 = Handle1->CastToRigidParticle())
+					{
+						Rigid1->SetCollisionConstraintFlag((uint32)Chaos::ECollisionConstraintFlags::CCF_BroadPhaseIgnoreCollisions);
+						DisabledCollisions.Add(Handle1->UniqueIdx());
+					}
+				}
 			}
-			PendingMap.Add(ActorReference, DisabledCollisions);
+
+			ActivationMap.Add(ActorIndex, DisabledCollisions);
 		}
 	}
 }
@@ -302,15 +321,18 @@ void FChaosEngineInterface::RemoveDisabledCollisionsFor_AssumesLocked(TArray< FP
 {
 	if (bEnableChaosCollisionManager)
 	{
-		for (FPhysicsActorHandle& Handle : InPhysicsActors)
+		for (FPhysicsActorHandle& ActorReference : InPhysicsActors)
 		{
-			Chaos::FPhysicsSolver* Solver = Handle->GetProxy()->GetSolver<Chaos::FPhysicsSolver>();
+			Chaos::FUniqueIdx ActorIndex = ActorReference->UniqueIdx();
+
+			Chaos::FPhysicsSolver* Solver = ActorReference->GetProxy()->GetSolver<Chaos::FPhysicsSolver>();
 			Chaos::FIgnoreCollisionManager& CollisionManager = Solver->GetEvolution()->GetBroadPhase().GetIgnoreCollisionManager();
 			int32 ExternalTimestamp = Solver->GetMarshallingManager().GetExternalTimestamp_External();
-			Chaos::FIgnoreCollisionManager::FDeactivationArray& PendingMap = CollisionManager.GetPendingDeactivationsForGameThread(ExternalTimestamp);
-			if (!PendingMap.Contains(Handle->UniqueIdx()))
+
+			Chaos::FIgnoreCollisionManager::FDeactivationArray& DeactivationMap = CollisionManager.GetPendingDeactivationsForGameThread(ExternalTimestamp);
+			if (!DeactivationMap.Contains(ActorReference->UniqueIdx()))
 			{
-				PendingMap.Add(Handle->UniqueIdx());
+				DeactivationMap.Add(ActorReference->UniqueIdx());
 			}
 		}
 	}
