Removed warning attempting to get scene for null actor reference. This is actually a valid case that simplifies scene locking. When creating constraints with the world one actor will be null and the rest of the API handles this gracefully. For an example see FBodyInstance::CreateDOFLock.

diff --git a/Engine/Source/Runtime/PhysicsCore/Private/ChaosEngineInterface.cpp b/Engine/Source/Runtime/PhysicsCore/Private/ChaosEngineInterface.cpp
--- a/Engine/Source/Runtime/PhysicsCore/Private/ChaosEngineInterface.cpp
+++ b/Engine/Source/Runtime/PhysicsCore/Private/ChaosEngineInterface.cpp
@@ -1575,8 +1575,6 @@ FChaosScene* FChaosEngineInterface::GetCurrentScene(const FPhysicsActorHandle& I
 {
 	if(!InHandle)
 	{
-		UE_LOG(LogChaos,Warning,TEXT("Attempting to get the current scene for a null handle."));
-		CHAOS_ENSURE(false);
 		return nullptr;
 	}
 
