Fixed radial impulse application to take into account the velocity change flag.

diff --git a/Engine/Source/Runtime/PhysicsCore/Private/ChaosEngineInterface.cpp b/Engine/Source/Runtime/PhysicsCore/Private/ChaosEngineInterface.cpp
--- a/Engine/Source/Runtime/PhysicsCore/Private/ChaosEngineInterface.cpp
+++ b/Engine/Source/Runtime/PhysicsCore/Private/ChaosEngineInterface.cpp
@@ -732,22 +732,21 @@ void FChaosEngineInterface::AddRadialImpulse_AssumesLocked(const FPhysicsActorHa
 		const FReal OriginToActorDistance = OriginToActor.Size();
 		if(OriginToActorDistance < InRadius)
 		{
+			FVec3 FinalImpulse = FVector::ZeroVector;
 			if(OriginToActorDistance > 0)
 			{
 				const FVec3 OriginToActorNorm = OriginToActor / OriginToActorDistance;
 
 				if(InFalloff == ERadialImpulseFalloff::RIF_Constant)
 				{
-					AddImpulse_AssumesLocked(InActorReference, OriginToActorNorm * InStrength);
-					return;
+					FinalImpulse = OriginToActorNorm * InStrength;
 				}
 				else if(InFalloff == ERadialImpulseFalloff::RIF_Linear)
 				{
 					const FReal DistanceOverlapping = InRadius - OriginToActorDistance;
 					if(DistanceOverlapping > 0)
 					{
-						FReal Strength = FMath::Lerp(0.0f, InStrength, DistanceOverlapping / InRadius);
-						AddImpulse_AssumesLocked(InActorReference, OriginToActorNorm * Strength);
+						FinalImpulse = OriginToActorNorm * FMath::Lerp(0.0f, InStrength, DistanceOverlapping / InRadius);
 					}
 				}
 				else
@@ -759,7 +758,16 @@ void FChaosEngineInterface::AddRadialImpulse_AssumesLocked(const FPhysicsActorHa
 			else
 			{
 				// Sphere and actor center are coincident, just pick a direction and apply maximum strength impulse.
-				AddImpulse_AssumesLocked(InActorReference, FVector::ForwardVector * InStrength);
+				FinalImpulse = FVector::ForwardVector * InStrength;
+			}
+
+			if(bInVelChange)
+			{
+				AddVelocity_AssumesLocked(InActorReference, FinalImpulse);
+			}
+			else
+			{
+				AddImpulse_AssumesLocked(InActorReference, FinalImpulse);
 			}
 		}
 	}
