Change UE4 reference in util comment

diff --git a/Engine/Source/Runtime/PhysicsCore/Public/PhysXToChaosUtil.h b/Engine/Source/Runtime/PhysicsCore/Public/PhysXToChaosUtil.h
--- a/Engine/Source/Runtime/PhysicsCore/Public/PhysXToChaosUtil.h
+++ b/Engine/Source/Runtime/PhysicsCore/Public/PhysXToChaosUtil.h
@@ -94,7 +94,7 @@ inline TUniquePtr<Chaos::TImplicitObjectTransformed<float, 3>> PxShapeToChaosGeo
 		MaterialIndices.Reserve(NumTris / 2);
 		for(int32 TriangleIndex = 0; TriangleIndex < NumTris; TriangleIndex += 2)
 		{
-			// We're only grabbing every other material because UE4 will only use one material per cell
+			// We're only grabbing every other material because Unreal will only use one material per cell
 			MaterialIndices.Add((uint8)(HeightField->getTriangleMaterialIndex(TriangleIndex)));
 		}
 		
