SetKinematicTarget_AssumesLocked only sets Target transform.
Set Previous trasform in FSingleParticlePhysicsProxy::PushToPhysicsState.
This solves the loss of momentum issue.

diff --git a/Engine/Source/Runtime/Experimental/Chaos/Private/PhysicsProxy/SingleParticlePhysicsProxy.cpp b/Engine/Source/Runtime/Experimental/Chaos/Private/PhysicsProxy/SingleParticlePhysicsProxy.cpp
--- a/Engine/Source/Runtime/Experimental/Chaos/Private/PhysicsProxy/SingleParticlePhysicsProxy.cpp
+++ b/Engine/Source/Runtime/Experimental/Chaos/Private/PhysicsProxy/SingleParticlePhysicsProxy.cpp
@@ -70,13 +70,17 @@ void PushToPhysicsStateImp(const Chaos::FDirtyPropertiesManager& Manager, Chaos:
 			KinematicHandle->SetVelocities(*NewVelocities);
 		}
 
-		auto NewKinematicTarget = bHasKinematicData ? ParticleData.FindKinematicTarget(Manager, DataIdx) : nullptr;
-		if (NewKinematicTarget)
+		auto NewKinematicTargetGT = bHasKinematicData ? ParticleData.FindKinematicTarget(Manager, DataIdx) : nullptr;
+		if (NewKinematicTargetGT)
 		{
-			KinematicHandle->SetKinematicTarget(*NewKinematicTarget);
+			const TRigidTransform<float,3> PreviousTransform(KinematicHandle->X(), KinematicHandle->R());
+			const TRigidTransform<float,3> &TargetTransform = NewKinematicTargetGT->GetTarget();
+			TKinematicTarget<float,3> NewKinematicTargetPT;
+			NewKinematicTargetPT.SetTargetMode(TargetTransform, PreviousTransform);
+			KinematicHandle->SetKinematicTarget(NewKinematicTargetPT);
 		}
 
-		if(NewXR || NewNonFrequentData || NewVelocities || NewKinematicTarget)
+		if(NewXR || NewNonFrequentData || NewVelocities || NewKinematicTargetGT)
 		{
 			const auto& Geometry = Handle->Geometry();
 			if(Geometry && Geometry->HasBoundingBox())
diff --git a/Engine/Source/Runtime/PhysicsCore/Private/ChaosEngineInterface.cpp b/Engine/Source/Runtime/PhysicsCore/Private/ChaosEngineInterface.cpp
--- a/Engine/Source/Runtime/PhysicsCore/Private/ChaosEngineInterface.cpp
+++ b/Engine/Source/Runtime/PhysicsCore/Private/ChaosEngineInterface.cpp
@@ -1621,13 +1621,18 @@ void FChaosEngineInterface::SetGlobalPose_AssumesLocked(const FPhysicsActorHandl
 void FChaosEngineInterface::SetKinematicTarget_AssumesLocked(const FPhysicsActorHandle& InActorReference,const FTransform& InNewTarget)
 {
 	{
-		Chaos::TKinematicTarget<float, 3> newKinematicTarget;
-		Chaos::FRigidTransform3 PreviousTM(InActorReference->GetGameThreadAPI().X(), InActorReference->GetGameThreadAPI().R());
-		newKinematicTarget.SetTargetMode(InNewTarget, PreviousTM);
-		InActorReference->GetGameThreadAPI().SetKinematicTarget(newKinematicTarget);
-
-		InActorReference->GetGameThreadAPI().SetX(InNewTarget.GetLocation());
-		InActorReference->GetGameThreadAPI().SetR(InNewTarget.GetRotation());
+		Chaos::TKinematicTarget<float, 3> NewKinematicTarget;
+
+		// SetKinematicTarget_AssumesLocked could be called multiple times in one time step
+		// Don't update Previous here. Previous is updated in FSingleParticlePhysicsProxy::PushToPhysicsState. Previous is not used in game thread and set to identity.
+		// @todo(chaos): create a new KinematicTarget class for game thread that does not have Preivous 
+		Chaos::TRigidTransform<Chaos::FReal, 3> IdentityTransform;
+		NewKinematicTarget.SetTargetMode(InNewTarget, IdentityTransform);
+		InActorReference->GetGameThreadAPI().SetKinematicTarget(NewKinematicTarget);
+
+		// Don't pass X/R to physics thread through SetX/SetR. KinematicTarget already does that.
+		InActorReference->GetGameThreadAPI().SetX(InNewTarget.GetLocation(), false);
+		InActorReference->GetGameThreadAPI().SetR(InNewTarget.GetRotation(), false);
 		InActorReference->GetGameThreadAPI().UpdateShapeBounds();
 
 		FChaosScene* Scene = GetCurrentScene(InActorReference);
