Undo //UE5/Release-5.0-EarlyAccess/Engine/Source/Runtime/... changelist 15897108

diff --git a/Engine/Source/Runtime/Experimental/Chaos/Private/PhysicsProxy/SingleParticlePhysicsProxy.cpp b/Engine/Source/Runtime/Experimental/Chaos/Private/PhysicsProxy/SingleParticlePhysicsProxy.cpp
--- a/Engine/Source/Runtime/Experimental/Chaos/Private/PhysicsProxy/SingleParticlePhysicsProxy.cpp
+++ b/Engine/Source/Runtime/Experimental/Chaos/Private/PhysicsProxy/SingleParticlePhysicsProxy.cpp
@@ -70,17 +70,13 @@ void PushToPhysicsStateImp(const Chaos::FDirtyPropertiesManager& Manager, Chaos:
 			KinematicHandle->SetVelocities(*NewVelocities);
 		}
 
-		auto NewKinematicTargetGT = bHasKinematicData ? ParticleData.FindKinematicTarget(Manager, DataIdx) : nullptr;
-		if (NewKinematicTargetGT)
+		auto NewKinematicTarget = bHasKinematicData ? ParticleData.FindKinematicTarget(Manager, DataIdx) : nullptr;
+		if (NewKinematicTarget)
 		{
-			const TRigidTransform<float,3> PreviousTransform(KinematicHandle->X(), KinematicHandle->R());
-			const TRigidTransform<float,3> &TargetTransform = NewKinematicTargetGT->GetTarget();
-			TKinematicTarget<float,3> NewKinematicTargetPT;
-			NewKinematicTargetPT.SetTargetMode(TargetTransform, PreviousTransform);
-			KinematicHandle->SetKinematicTarget(NewKinematicTargetPT);
+			KinematicHandle->SetKinematicTarget(*NewKinematicTarget);
 		}
 
-		if(NewXR || NewNonFrequentData || NewVelocities || NewKinematicTargetGT)
+		if(NewXR || NewNonFrequentData || NewVelocities || NewKinematicTarget)
 		{
 			const auto& Geometry = Handle->Geometry();
 			if(Geometry && Geometry->HasBoundingBox())
diff --git a/Engine/Source/Runtime/PhysicsCore/Private/ChaosEngineInterface.cpp b/Engine/Source/Runtime/PhysicsCore/Private/ChaosEngineInterface.cpp
--- a/Engine/Source/Runtime/PhysicsCore/Private/ChaosEngineInterface.cpp
+++ b/Engine/Source/Runtime/PhysicsCore/Private/ChaosEngineInterface.cpp
@@ -1621,18 +1621,13 @@ void FChaosEngineInterface::SetGlobalPose_AssumesLocked(const FPhysicsActorHandl
 void FChaosEngineInterface::SetKinematicTarget_AssumesLocked(const FPhysicsActorHandle& InActorReference,const FTransform& InNewTarget)
 {
 	{
-		Chaos::TKinematicTarget<float, 3> NewKinematicTarget;
-
-		// SetKinematicTarget_AssumesLocked could be called multiple times in one time step
-		// Don't update Previous here. Previous is updated in FSingleParticlePhysicsProxy::PushToPhysicsState. Previous is not used in game thread and set to identity.
-		// @todo(chaos): create a new KinematicTarget class for game thread that does not have Preivous 
-		Chaos::TRigidTransform<Chaos::FReal, 3> IdentityTransform;
-		NewKinematicTarget.SetTargetMode(InNewTarget, IdentityTransform);
-		InActorReference->GetGameThreadAPI().SetKinematicTarget(NewKinematicTarget);
-
-		// Don't pass X/R to physics thread through SetX/SetR. KinematicTarget already does that.
-		InActorReference->GetGameThreadAPI().SetX(InNewTarget.GetLocation(), false);
-		InActorReference->GetGameThreadAPI().SetR(InNewTarget.GetRotation(), false);
+		Chaos::TKinematicTarget<float, 3> newKinematicTarget;
+		Chaos::TRigidTransform<Chaos::FReal, 3> PreviousTM(InActorReference->GetGameThreadAPI().X(), InActorReference->GetGameThreadAPI().R());
+		newKinematicTarget.SetTargetMode(InNewTarget, PreviousTM);
+		InActorReference->GetGameThreadAPI().SetKinematicTarget(newKinematicTarget);
+
+		InActorReference->GetGameThreadAPI().SetX(InNewTarget.GetLocation());
+		InActorReference->GetGameThreadAPI().SetR(InNewTarget.GetRotation());
 		InActorReference->GetGameThreadAPI().UpdateShapeBounds();
 
 		FChaosScene* Scene = GetCurrentScene(InActorReference);
