Chaos Physics: When changing from Kinematic to Dynamic and vice versa  make sure that the velocities are reset.

diff --git a/Engine/Source/Runtime/PhysicsCore/Private/ChaosEngineInterface.cpp b/Engine/Source/Runtime/PhysicsCore/Private/ChaosEngineInterface.cpp
--- a/Engine/Source/Runtime/PhysicsCore/Private/ChaosEngineInterface.cpp
+++ b/Engine/Source/Runtime/PhysicsCore/Private/ChaosEngineInterface.cpp
@@ -485,6 +485,12 @@ void FChaosEngineInterface::SetIsKinematic_AssumesLocked(const FPhysicsActorHand
 			{
 				InActorReference->GetGameThreadAPI().SetResimType(EResimType::FullResim);
 			}
+			else if (NewState == Chaos::EObjectStateType::Kinematic)
+			{
+				// Reset velocity on a state change here
+				InActorReference->GetGameThreadAPI().SetV(Chaos::FVec3((Chaos::FReal) 0));
+				InActorReference->GetGameThreadAPI().SetW(Chaos::FVec3((Chaos::FReal) 0));
+			}
 		}
 	}
 }
