Fix to make PS5 build when LWC is on

diff --git a/Engine/Plugins/Runtime/HairStrands/Source/HairStrandsCore/Private/GroomBuilder.cpp b/Engine/Plugins/Runtime/HairStrands/Source/HairStrandsCore/Private/GroomBuilder.cpp
--- a/Engine/Plugins/Runtime/HairStrands/Source/HairStrandsCore/Private/GroomBuilder.cpp
+++ b/Engine/Plugins/Runtime/HairStrands/Source/HairStrandsCore/Private/GroomBuilder.cpp
@@ -1603,11 +1603,11 @@ bool FGroomBuilder::BuildHairDescriptionGroups(const FHairDescription& HairDescr
 
 	// Compute bound radius
 	{
-		FVector GroomBoundMin(FLT_MAX);
-		FVector GroomBoundMax(-FLT_MAX);
+		FVector3f GroomBoundMin(FLT_MAX);
+		FVector3f GroomBoundMax(-FLT_MAX);
 		for (const FHairDescriptionGroup& Group : Out.HairGroups)
 		{
-			for (const FVector& P : Group.Strands.StrandsPoints.PointsPosition)
+			for (const FVector3f& P : Group.Strands.StrandsPoints.PointsPosition)
 			{
 				GroomBoundMin.X = FMath::Min(GroomBoundMin.X, P.X);
 				GroomBoundMin.Y = FMath::Min(GroomBoundMin.Y, P.Y);
@@ -1619,7 +1619,7 @@ bool FGroomBuilder::BuildHairDescriptionGroups(const FHairDescription& HairDescr
 			}
 		}
 
-		Out.BoundRadius = FVector::Distance(GroomBoundMax, GroomBoundMin) * 0.5f;
+		Out.BoundRadius = FVector3f::Distance(GroomBoundMax, GroomBoundMin) * 0.5f;
 	}
 
 	// Update GroupInfo
diff --git a/Engine/Source/Runtime/PhysicsCore/Public/SQVerifier.h b/Engine/Source/Runtime/PhysicsCore/Public/SQVerifier.h
--- a/Engine/Source/Runtime/PhysicsCore/Public/SQVerifier.h
+++ b/Engine/Source/Runtime/PhysicsCore/Public/SQVerifier.h
@@ -41,7 +41,7 @@ void SQPerfComparisonHelper(const FString& TestName, FPhysTestSerializer& Serial
 		}
 #endif
 
-		ISpatialAccelerationCollection<FAccelerationStructureHandle, float, 3>* Accelerator = nullptr;
+		ISpatialAccelerationCollection<FAccelerationStructureHandle, FReal, 3>* Accelerator = nullptr;
 		Serializer.GetChaosData()->UpdateExternalAccelerationStructure_External(Accelerator, Empty);
 		FChaosSQAccelerator SQAccelerator(*Accelerator);
 		for (double i = 0; i < NumIterations; ++i)
@@ -68,7 +68,7 @@ void SQPerfComparisonHelper(const FString& TestName, FPhysTestSerializer& Serial
 		}
 #endif
 
-		ISpatialAccelerationCollection<FAccelerationStructureHandle, float, 3>* Accelerator = nullptr;
+		ISpatialAccelerationCollection<FAccelerationStructureHandle, FReal, 3>* Accelerator = nullptr;
 		Serializer.GetChaosData()->UpdateExternalAccelerationStructure_External(Accelerator, Empty);
 		FChaosSQAccelerator SQAccelerator(*Accelerator);
 		for (double i = 0; i < NumIterations; ++i)
@@ -95,7 +95,7 @@ void SQPerfComparisonHelper(const FString& TestName, FPhysTestSerializer& Serial
 		}
 #endif
 
-		ISpatialAccelerationCollection<FAccelerationStructureHandle, float, 3>* Accelerator = nullptr;
+		ISpatialAccelerationCollection<FAccelerationStructureHandle, FReal, 3>* Accelerator = nullptr;
 		Serializer.GetChaosData()->UpdateExternalAccelerationStructure_External(Accelerator, Empty);
 		FChaosSQAccelerator SQAccelerator(*Accelerator);
 		for (double i = 0; i < NumIterations; ++i)
