Remove ClearKinematicTarget

diff --git a/Engine/Source/Programs/HeadlessChaos/Private/HeadlessChaosTestPhysicsScene.cpp b/Engine/Source/Programs/HeadlessChaos/Private/HeadlessChaosTestPhysicsScene.cpp
--- a/Engine/Source/Programs/HeadlessChaos/Private/HeadlessChaosTestPhysicsScene.cpp
+++ b/Engine/Source/Programs/HeadlessChaos/Private/HeadlessChaosTestPhysicsScene.cpp
@@ -1135,9 +1135,6 @@ namespace ChaosTest {
 		FChaosEngineInterface::SetKinematicTarget_AssumesLocked(Proxy, FTransform(CurrentX));
 		FChaosEngineInterface::SetGlobalPose_AssumesLocked(Proxy, FTransform(CurrentX));
 
-		// Test if KinematicTarget is cleared
-		EXPECT_EQ(Particle.IsKinematicTargetDirty(), false);
-
 		Callback->GetProducerInputData_External()->CorrectX = CurrentX;
 		AdvanceFrameAndRunTest(CurrentX, CurrentV);
 
diff --git a/Engine/Source/Runtime/PhysicsCore/Private/ChaosEngineInterface.cpp b/Engine/Source/Runtime/PhysicsCore/Private/ChaosEngineInterface.cpp
--- a/Engine/Source/Runtime/PhysicsCore/Private/ChaosEngineInterface.cpp
+++ b/Engine/Source/Runtime/PhysicsCore/Private/ChaosEngineInterface.cpp
@@ -1662,10 +1662,6 @@ FChaosScene* FChaosEngineInterface::GetCurrentScene(const FPhysicsActorHandle& I
 void FChaosEngineInterface::SetGlobalPose_AssumesLocked(const FPhysicsActorHandle& InActorReference,const FTransform& InNewPose,bool bAutoWake)
 {
 	Chaos::FRigidBodyHandle_External& Body_External = InActorReference->GetGameThreadAPI();
-	if (Body_External.IsKinematicTargetDirty())
-	{
-		Body_External.ClearKinematicTarget();
-	}
 	Body_External.SetX(InNewPose.GetLocation());
 	Body_External.SetR(InNewPose.GetRotation());
 	Body_External.UpdateShapeBounds();
