Initialize solver AsyncDt from project settings on solver init, instead of setting AsyncDt every frame.
Fixes project setting stomping over AsyncDt modifications from other sources like playlists.
Also fixes logspam of async dt.

diff --git a/Engine/Source/Runtime/Engine/Private/PhysicsEngine/Experimental/PhysScene_Chaos.cpp b/Engine/Source/Runtime/Engine/Private/PhysicsEngine/Experimental/PhysScene_Chaos.cpp
--- a/Engine/Source/Runtime/Engine/Private/PhysicsEngine/Experimental/PhysScene_Chaos.cpp
+++ b/Engine/Source/Runtime/Engine/Private/PhysicsEngine/Experimental/PhysScene_Chaos.cpp
@@ -372,6 +372,7 @@ FPhysScene_Chaos::FPhysScene_Chaos(AActor* InSolverActor
 #endif
 )
 	: Super(InSolverActor ? InSolverActor->GetWorld() : nullptr
+		, UPhysicsSettings::Get()->bTickPhysicsAsync ? UPhysicsSettings::Get()->AsyncFixedTimeStepSize : -1
 #if CHAOS_CHECKED
 		, DebugName
 #endif
diff --git a/Engine/Source/Runtime/Engine/Private/PhysicsEngine/PhysLevel.cpp b/Engine/Source/Runtime/Engine/Private/PhysicsEngine/PhysLevel.cpp
--- a/Engine/Source/Runtime/Engine/Private/PhysicsEngine/PhysLevel.cpp
+++ b/Engine/Source/Runtime/Engine/Private/PhysicsEngine/PhysLevel.cpp
@@ -144,7 +144,7 @@ void UWorld::SetupPhysicsTickFunctions(float DeltaSeconds)
 
 	static const auto CVar_MaxPhysicsDeltaTime = IConsoleManager::Get().FindTConsoleVariableDataFloat(TEXT("p.MaxPhysicsDeltaTime"));
 	PhysScene->SetUpForFrame(&DefaultGravity, DeltaSeconds, UPhysicsSettings::Get()->MaxPhysicsDeltaTime,
-		UPhysicsSettings::Get()->MaxSubstepDeltaTime, UPhysicsSettings::Get()->MaxSubsteps, UPhysicsSettings::Get()->bSubstepping, UPhysicsSettings::Get()->bTickPhysicsAsync ? UPhysicsSettings::Get()->AsyncFixedTimeStepSize : -1);
+		UPhysicsSettings::Get()->MaxSubstepDeltaTime, UPhysicsSettings::Get()->MaxSubsteps, UPhysicsSettings::Get()->bSubstepping);
 }
 
 void UWorld::StartPhysicsSim()
diff --git a/Engine/Source/Runtime/Engine/Private/World.cpp b/Engine/Source/Runtime/Engine/Private/World.cpp
--- a/Engine/Source/Runtime/Engine/Private/World.cpp
+++ b/Engine/Source/Runtime/Engine/Private/World.cpp
@@ -1670,7 +1670,7 @@ void UWorld::InitWorld(const InitializationValues IVS)
 	if (GetPhysicsScene())
 	{
 		FVector Gravity = FVector( 0.f, 0.f, GetGravityZ() );
-		GetPhysicsScene()->SetUpForFrame( &Gravity, 0, 0, 0, 0, false, -1 );
+		GetPhysicsScene()->SetUpForFrame( &Gravity, 0, 0, 0, 0, false);
 	}
 
 	// Create physics collision handler, if we have a physics scene
diff --git a/Engine/Source/Runtime/Experimental/Chaos/Private/Chaos/Framework/PhysicsSolverBase.cpp b/Engine/Source/Runtime/Experimental/Chaos/Private/Chaos/Framework/PhysicsSolverBase.cpp
--- a/Engine/Source/Runtime/Experimental/Chaos/Private/Chaos/Framework/PhysicsSolverBase.cpp
+++ b/Engine/Source/Runtime/Experimental/Chaos/Private/Chaos/Framework/PhysicsSolverBase.cpp
@@ -126,7 +126,7 @@ namespace Chaos
 		" 1 blocks on everything except the single most recent task (including tasks from current frame). 1 should gurantee we will always have a future output for interpolation from 2 frames in the past."));
 
 
-	FPhysicsSolverBase::FPhysicsSolverBase(const EMultiBufferMode BufferingModeIn,const EThreadingModeTemp InThreadingMode,UObject* InOwner)
+	FPhysicsSolverBase::FPhysicsSolverBase(const EMultiBufferMode BufferingModeIn,const EThreadingModeTemp InThreadingMode,UObject* InOwner, Chaos::FReal InAsyncDt)
 		: BufferMode(BufferingModeIn)
 		, ThreadingMode(InThreadingMode)
 		, PullResultsManager(MakeUnique<FChaosResultsManager>(MarshallingManager))
@@ -136,7 +136,7 @@ namespace Chaos
 		, Owner(InOwner)
 		, ExternalDataLock_External(new FPhysicsSceneGuard())
 		, bIsShuttingDown(false)
-		, AsyncDt(-1)
+		, AsyncDt(InAsyncDt)
 		, AccumulatedTime(0)
 		, MMaxDeltaTime(0.0)
 		, MMinDeltaTime(SMALL_NUMBER)
diff --git a/Engine/Source/Runtime/Experimental/Chaos/Private/ChaosSolversModule.cpp b/Engine/Source/Runtime/Experimental/Chaos/Private/ChaosSolversModule.cpp
--- a/Engine/Source/Runtime/Experimental/Chaos/Private/ChaosSolversModule.cpp
+++ b/Engine/Source/Runtime/Experimental/Chaos/Private/ChaosSolversModule.cpp
@@ -180,7 +180,7 @@ void FChaosSolversModule::SyncTask(bool bForceBlockingSync /*= false*/)
 #endif
 }
 
-Chaos::FPBDRigidsSolver* FChaosSolversModule::CreateSolver(UObject* InOwner, Chaos::EThreadingMode InThreadingMode
+Chaos::FPBDRigidsSolver* FChaosSolversModule::CreateSolver(UObject* InOwner, Chaos::FReal InAsyncDt, Chaos::EThreadingMode InThreadingMode
 #if CHAOS_CHECKED
 	, const FName& DebugName
 #endif
@@ -192,8 +192,7 @@ Chaos::FPBDRigidsSolver* FChaosSolversModule::CreateSolver(UObject* InOwner, Cha
 	FChaosScopeSolverLock SolverScopeLock;
 	
 	EMultiBufferMode SolverBufferMode = InThreadingMode == EThreadingMode::SingleThread ? EMultiBufferMode::Single : EMultiBufferMode::Double;
-	
-	auto* NewSolver = new FPBDRigidsSolver(SolverBufferMode,InOwner);
+	auto* NewSolver = new FPBDRigidsSolver(SolverBufferMode,InOwner, InAsyncDt);
 	AllSolvers.Add(NewSolver);
 
 	// Add The solver to the owner list
diff --git a/Engine/Source/Runtime/Experimental/Chaos/Private/PBDRigidsSolver.cpp b/Engine/Source/Runtime/Experimental/Chaos/Private/PBDRigidsSolver.cpp
--- a/Engine/Source/Runtime/Experimental/Chaos/Private/PBDRigidsSolver.cpp
+++ b/Engine/Source/Runtime/Experimental/Chaos/Private/PBDRigidsSolver.cpp
@@ -390,8 +390,8 @@ namespace Chaos
 		TSharedPtr<FEvent> PrevEvent, CurrentEvent;
 	};
 
-	FPBDRigidsSolver::FPBDRigidsSolver(const EMultiBufferMode BufferingModeIn, UObject* InOwner)
-		: Super(BufferingModeIn, BufferingModeIn == EMultiBufferMode::Single ? EThreadingModeTemp::SingleThread : EThreadingModeTemp::TaskGraph, InOwner)
+	FPBDRigidsSolver::FPBDRigidsSolver(const EMultiBufferMode BufferingModeIn, UObject* InOwner, FReal InAsyncDt)
+		: Super(BufferingModeIn, BufferingModeIn == EMultiBufferMode::Single ? EThreadingModeTemp::SingleThread : EThreadingModeTemp::TaskGraph, InOwner, InAsyncDt)
 		, CurrentFrame(0)
 		, MTime(0.0)
 		, MLastDt(0.0)
diff --git a/Engine/Source/Runtime/Experimental/Chaos/Public/Chaos/Framework/PhysicsSolverBase.h b/Engine/Source/Runtime/Experimental/Chaos/Public/Chaos/Framework/PhysicsSolverBase.h
--- a/Engine/Source/Runtime/Experimental/Chaos/Public/Chaos/Framework/PhysicsSolverBase.h
+++ b/Engine/Source/Runtime/Experimental/Chaos/Public/Chaos/Framework/PhysicsSolverBase.h
@@ -399,7 +399,7 @@ namespace Chaos
 		EThreadingModeTemp ThreadingMode;
 
 		/** Protected construction so callers still have to go through the module to create new instances */
-		FPhysicsSolverBase(const EMultiBufferMode BufferingModeIn,const EThreadingModeTemp InThreadingMode,UObject* InOwner);
+		FPhysicsSolverBase(const EMultiBufferMode BufferingModeIn,const EThreadingModeTemp InThreadingMode,UObject* InOwner, FReal AsyncDt);
 
 		/** Only allow construction with valid parameters as well as restricting to module construction */
 		virtual ~FPhysicsSolverBase();
diff --git a/Engine/Source/Runtime/Experimental/Chaos/Public/ChaosSolversModule.h b/Engine/Source/Runtime/Experimental/Chaos/Public/ChaosSolversModule.h
--- a/Engine/Source/Runtime/Experimental/Chaos/Public/ChaosSolversModule.h
+++ b/Engine/Source/Runtime/Experimental/Chaos/Public/ChaosSolversModule.h
@@ -126,7 +126,7 @@ public:
 	 * @param ThreadingMode The desired threading mode the solver will use
 	 * @param bStandalone Whether the solver is standalone (not sent to physics thread - updating left to caller)
 	 */
-	Chaos::FPBDRigidsSolver* CreateSolver(UObject* InOwner, Chaos::EThreadingMode ThreadingMode = Chaos::EThreadingMode::SingleThread
+	Chaos::FPBDRigidsSolver* CreateSolver(UObject* InOwner, Chaos::FReal InAsyncDt, Chaos::EThreadingMode ThreadingMode = Chaos::EThreadingMode::SingleThread
 #if CHAOS_CHECKED
 		, const FName& DebugName = NAME_None
 #endif
diff --git a/Engine/Source/Runtime/Experimental/Chaos/Public/PBDRigidsSolver.h b/Engine/Source/Runtime/Experimental/Chaos/Public/PBDRigidsSolver.h
--- a/Engine/Source/Runtime/Experimental/Chaos/Public/PBDRigidsSolver.h
+++ b/Engine/Source/Runtime/Experimental/Chaos/Public/PBDRigidsSolver.h
@@ -82,7 +82,7 @@ namespace Chaos
 	class CHAOS_API FPBDRigidsSolver : public FPhysicsSolverBase
 	{
 
-		FPBDRigidsSolver(const EMultiBufferMode BufferingModeIn, UObject* InOwner);
+		FPBDRigidsSolver(const EMultiBufferMode BufferingModeIn, UObject* InOwner, FReal AsyncDt);
 		virtual ~FPBDRigidsSolver();
 
 	public:
diff --git a/Engine/Source/Runtime/PhysicsCore/Private/ChaosScene.cpp b/Engine/Source/Runtime/PhysicsCore/Private/ChaosScene.cpp
--- a/Engine/Source/Runtime/PhysicsCore/Private/ChaosScene.cpp
+++ b/Engine/Source/Runtime/PhysicsCore/Private/ChaosScene.cpp
@@ -46,6 +46,7 @@ TAutoConsoleVariable<int32> CVar_ApplyProjectSettings(TEXT("p.Chaos.Simulation.A
 
 FChaosScene::FChaosScene(
 	UObject* OwnerPtr
+	, Chaos::FReal InAsyncDt
 #if CHAOS_CHECKED
 	, const FName& DebugName
 #endif
@@ -64,7 +65,7 @@ FChaosScene::FChaosScene(
 
 	Chaos::EThreadingMode ThreadingMode = bForceSingleThread ? Chaos::EThreadingMode::SingleThread : Chaos::EThreadingMode::TaskGraph;
 
-	SceneSolver = ChaosModule->CreateSolver(OwnerPtr,ThreadingMode
+	SceneSolver = ChaosModule->CreateSolver(OwnerPtr, InAsyncDt, ThreadingMode
 #if CHAOS_CHECKED
 		,DebugName
 #endif
@@ -321,7 +322,7 @@ void FChaosScene::SetGravity(const Chaos::FVec3& Acceleration)
 	SimCallback->GetProducerInputData_External()->Gravity = Acceleration;
 }
 
-void FChaosScene::SetUpForFrame(const FVector* NewGrav,float InDeltaSeconds /*= 0.0f*/,float InMaxPhysicsDeltaTime /*= 0.0f*/,float InMaxSubstepDeltaTime /*= 0.0f*/,int32 InMaxSubsteps,bool bSubstepping,float InAsyncDt)
+void FChaosScene::SetUpForFrame(const FVector* NewGrav,float InDeltaSeconds /*= 0.0f*/,float InMaxPhysicsDeltaTime /*= 0.0f*/,float InMaxSubstepDeltaTime /*= 0.0f*/,int32 InMaxSubsteps,bool bSubstepping)
 {
 #if WITH_CHAOS
 	using namespace Chaos;
@@ -338,15 +339,6 @@ void FChaosScene::SetUpForFrame(const FVector* NewGrav,float InDeltaSeconds /*=
 
 	if(FPhysicsSolver* Solver = GetSolver())
 	{
-		const bool bIsAsync = InAsyncDt > 0;
-		if (bIsAsync)
-		{
-			Solver->EnableAsyncMode(InAsyncDt);
-		}
-		else
-		{
-			Solver->DisableAsyncMode();
-		}
 		if(bSubstepping)
 		{
 			Solver->SetMaxDeltaTime_External(InMaxSubstepDeltaTime);
diff --git a/Engine/Source/Runtime/PhysicsCore/Public/Chaos/ChaosScene.h b/Engine/Source/Runtime/PhysicsCore/Public/Chaos/ChaosScene.h
--- a/Engine/Source/Runtime/PhysicsCore/Public/Chaos/ChaosScene.h
+++ b/Engine/Source/Runtime/PhysicsCore/Public/Chaos/ChaosScene.h
@@ -80,6 +80,7 @@ class PHYSICSCORE_API FChaosScene
 public:
 	FChaosScene(
 		UObject* OwnerPtr
+		, Chaos::FReal InAsyncDt
 #if CHAOS_CHECKED
 	, const FName& DebugName=NAME_None
 #endif
@@ -132,7 +133,7 @@ public:
 #endif
 
 	void StartFrame();
-	void SetUpForFrame(const FVector* NewGrav,float InDeltaSeconds,float InMaxPhysicsDeltaTime,float InMaxSubstepDeltaTime,int32 InMaxSubsteps,bool bSubstepping,float InAsyncDt);
+	void SetUpForFrame(const FVector* NewGrav,float InDeltaSeconds,float InMaxPhysicsDeltaTime,float InMaxSubstepDeltaTime,int32 InMaxSubsteps,bool bSubstepping);
 	void EndFrame();
 
 	DECLARE_MULTICAST_DELEGATE_OneParam(FOnPhysScenePostTick,FChaosScene*);
