Chaos: Adding some statistics to AABBTree

#[fyi] Michael.Lentine, Benn.Gallagherl
#[review]-16743248 [REDACTED_HANDLE]

diff --git a/Engine/Source/Runtime/Experimental/Chaos/Public/Chaos/AABBTree.h b/Engine/Source/Runtime/Experimental/Chaos/Public/Chaos/AABBTree.h
--- a/Engine/Source/Runtime/Experimental/Chaos/Public/Chaos/AABBTree.h
+++ b/Engine/Source/Runtime/Experimental/Chaos/Public/Chaos/AABBTree.h
@@ -45,6 +45,31 @@ enum class EAABBQueryType
 	Overlap
 };
 
+struct AABBTreeStatistics
+{
+	void Reset()
+	{
+		StatNumNonEmptyCellsInGrid = 0;
+		StatNumElementsTooLargeForGrid = 0;
+		StatNumDirtyElements = 0;
+		StatNumGridOverflowElements = 0;
+	}
+
+	AABBTreeStatistics& operator+=(const AABBTreeStatistics& Rhs)
+	{
+		StatNumNonEmptyCellsInGrid += Rhs.StatNumNonEmptyCellsInGrid;
+		StatNumElementsTooLargeForGrid += Rhs.StatNumElementsTooLargeForGrid;
+		StatNumDirtyElements += Rhs.StatNumDirtyElements;
+		StatNumGridOverflowElements += Rhs.StatNumGridOverflowElements;
+		return *this;
+	}
+
+	int32 StatNumNonEmptyCellsInGrid = 0;
+	int32 StatNumElementsTooLargeForGrid = 0;
+	int32 StatNumDirtyElements = 0;
+	int32 StatNumGridOverflowElements = 0;
+};
+
 DECLARE_CYCLE_STAT(TEXT("AABBTreeGenerateTree"), STAT_AABBTreeGenerateTree, STATGROUP_Chaos);
 DECLARE_CYCLE_STAT(TEXT("AABBTreeTimeSliceSetup"), STAT_AABBTreeTimeSliceSetup, STATGROUP_Chaos);
 DECLARE_CYCLE_STAT(TEXT("AABBTreeInitialTimeSlice"), STAT_AABBTreeInitialTimeSlice, STATGROUP_Chaos);
@@ -400,6 +425,7 @@ public:
 		CellHashToFlatArray.Reset();
 		FlattenedCellArrayOfDirtyIndices.Reset();
 		DirtyElementsGridOverflow.Reset();
+		TreeStats.Reset();
 		GlobalPayloads.Reset();
 		PayloadToInfo.Reset();
 
@@ -594,6 +620,7 @@ public:
 			NewHashEntry.Count = 1;
 			FlattenedCellArrayOfDirtyIndices.AddUninitialized(DirtyElementMaxCellCapacity);
 			FlattenedCellArrayOfDirtyIndices[NewHashEntry.Index] = NewDirtyIndex;
+			TreeStats.StatNumNonEmptyCellsInGrid++;
 			return true;
 		}
 		return false;
@@ -668,6 +695,11 @@ public:
 	FORCEINLINE_DEBUGGABLE int32 AddDirtyElementToGrid(const TAABB<FReal, 3>& NewBounds, int32 NewDirtyElement)
 	{
 		bool bAddToGrid = !TooManyOverlapQueryCells(NewBounds, DirtyElementGridCellSizeInv, DirtyElementMaxPhysicalSizeInCells);
+		if (!bAddToGrid)
+		{
+			TreeStats.StatNumElementsTooLargeForGrid++;
+		}
+
 		if (bAddToGrid)
 		{
 			DoForOverlappedCells(NewBounds, DirtyElementGridCellSize, DirtyElementGridCellSizeInv, [&](int32 Hash) {
@@ -891,6 +923,14 @@ public:
 		return DirtyElements.Num();
 	}
 
+	// Some useful statistics
+	const AABBTreeStatistics& GetAABBTreeStatistics() {
+		// Update the stats that needs it first
+		TreeStats.StatNumDirtyElements = DirtyElements.Num();
+		TreeStats.StatNumGridOverflowElements = DirtyElementsGridOverflow.Num();
+		return TreeStats;
+	}
+
 	const TArray<TPayloadBoundsElement<TPayloadType, FReal>>& GlobalObjects() const
 	{
 		return GlobalPayloads;
@@ -1297,6 +1337,7 @@ private:
 		CellHashToFlatArray.Reset(); 
 		FlattenedCellArrayOfDirtyIndices.Reset();
 		DirtyElementsGridOverflow.Reset();
+		TreeStats.Reset();
 		PayloadToInfo.Reset();
 		NumProcessedThisSlice = 0;
 		GetCVars();  // Safe to copy CVARS here
@@ -1693,6 +1734,7 @@ private:
 		, DirtyElementMaxGridCellQueryCount(Other.DirtyElementMaxGridCellQueryCount)
 		, DirtyElementMaxPhysicalSizeInCells(Other.DirtyElementMaxPhysicalSizeInCells)
 		, DirtyElementMaxCellCapacity(Other.DirtyElementMaxCellCapacity)
+		, TreeStats(Other.TreeStats)
 		, GlobalPayloads(Other.GlobalPayloads)
 		, PayloadToInfo(Other.PayloadToInfo)
 		, MaxChildrenInLeaf(Other.MaxChildrenInLeaf)
@@ -1721,6 +1763,7 @@ private:
 			CellHashToFlatArray = Rhs.CellHashToFlatArray;
 			FlattenedCellArrayOfDirtyIndices = Rhs.FlattenedCellArrayOfDirtyIndices;
 			DirtyElementsGridOverflow = Rhs.DirtyElementsGridOverflow;
+			TreeStats = Rhs.TreeStats;
 			
 			DirtyElementGridCellSize = Rhs.DirtyElementGridCellSize;
 			DirtyElementGridCellSizeInv = Rhs.DirtyElementGridCellSizeInv;
@@ -1754,6 +1797,8 @@ private:
 	int32 DirtyElementMaxGridCellQueryCount;
 	int32 DirtyElementMaxPhysicalSizeInCells;
 	int32 DirtyElementMaxCellCapacity;
+	// Some useful statistics
+	AABBTreeStatistics TreeStats;
 
 	TArray<FElement> GlobalPayloads;
 	TArrayAsMap<TPayloadType, FAABBTreePayloadInfo> PayloadToInfo;
diff --git a/Engine/Source/Runtime/PhysicsCore/Private/ChaosScene.cpp b/Engine/Source/Runtime/PhysicsCore/Private/ChaosScene.cpp
--- a/Engine/Source/Runtime/PhysicsCore/Private/ChaosScene.cpp
+++ b/Engine/Source/Runtime/PhysicsCore/Private/ChaosScene.cpp
@@ -40,6 +40,9 @@ CSV_DEFINE_CATEGORY(ChaosPhysics,true);
 
 // Stat Counters
 DECLARE_DWORD_ACCUMULATOR_STAT(TEXT("NumDirtyAABBTreeElements"), STAT_ChaosCounter_NumDirtyAABBTreeElements, STATGROUP_ChaosCounters);
+DECLARE_DWORD_ACCUMULATOR_STAT(TEXT("NumDirtyGridOverflowElements"), STAT_ChaosCounter_NumDirtyGridOverflowElements, STATGROUP_ChaosCounters);
+DECLARE_DWORD_ACCUMULATOR_STAT(TEXT("NumDirtyElementsTooLargeForGrid"), STAT_ChaosCounter_NumDirtyElementsTooLargeForGrid, STATGROUP_ChaosCounters);
+DECLARE_DWORD_ACCUMULATOR_STAT(TEXT("NumDirtyNonEmptyCellsInGrid"), STAT_ChaosCounter_NumDirtyNonEmptyCellsInGrid, STATGROUP_ChaosCounters);
 
 TAutoConsoleVariable<int32> CVar_ChaosSimulationEnable(TEXT("P.Chaos.Simulation.Enable"),1,TEXT("Enable / disable chaos simulation. If disabled, physics will not tick."));
 TAutoConsoleVariable<int32> CVar_ApplyProjectSettings(TEXT("p.Chaos.Simulation.ApplySolverProjectSettings"), 1, TEXT("Whether to apply the solver project settings on spawning a solver"));
@@ -432,27 +435,24 @@ void FChaosScene::SyncBodies(TSolver* Solver)
 #endif
 }
 
-
-// Find the number of dirty elements in all substructures that has dirty elements that we know of
-// This is non recursive for now
-// Todo: consider making DirtyElementsCount a method on ISpatialAcceleration instead
-int32 DirtyElementCount(Chaos::ISpatialAccelerationCollection<Chaos::FAccelerationStructureHandle,Chaos::FReal,3>& Collection)
+// Accumulate all the AABBTree stats
+void GetAABBTreeStats(Chaos::ISpatialAccelerationCollection<Chaos::FAccelerationStructureHandle, Chaos::FReal, 3>& Collection, Chaos::AABBTreeStatistics& OutAABBTreeStatistics)
 {
 	using namespace Chaos;
-	int32 DirtyElements = 0;
+	OutAABBTreeStatistics.Reset();
 	TArray<FSpatialAccelerationIdx> SpatialIndices = Collection.GetAllSpatialIndices();
-	for(const FSpatialAccelerationIdx SpatialIndex : SpatialIndices)
+	for (const FSpatialAccelerationIdx SpatialIndex : SpatialIndices)
 	{
 		auto SubStructure = Collection.GetSubstructure(SpatialIndex);
-		if(const auto AABBTree = SubStructure->template As<TAABBTree<FAccelerationStructureHandle,TAABBTreeLeafArray<FAccelerationStructureHandle>>>())
+		if (const auto AABBTree = SubStructure->template As<TAABBTree<FAccelerationStructureHandle, TAABBTreeLeafArray<FAccelerationStructureHandle>>>())
 		{
-			DirtyElements += AABBTree->NumDirtyElements();
-		} else if(const auto AABBTreeBV = SubStructure->template As<TAABBTree<FAccelerationStructureHandle,TBoundingVolume<FAccelerationStructureHandle>>>())
+			OutAABBTreeStatistics += AABBTree->GetAABBTreeStatistics();
+		}
+		else if (const auto AABBTreeBV = SubStructure->template As<TAABBTree<FAccelerationStructureHandle, TBoundingVolume<FAccelerationStructureHandle>>>())
 		{
-			DirtyElements += AABBTreeBV->NumDirtyElements();
+			OutAABBTreeStatistics += AABBTreeBV->GetAABBTreeStatistics();
 		}
 	}
-	return DirtyElements;
 }
 
 void FChaosScene::EndFrame()
@@ -468,9 +468,20 @@ void FChaosScene::EndFrame()
 		return;
 	}
 
-	int32 DirtyElements = DirtyElementCount(GetSpacialAcceleration()->AsChecked<SpatialAccelerationCollection>());
-	CSV_CUSTOM_STAT(ChaosPhysics,AABBTreeDirtyElementCount,DirtyElements,ECsvCustomStatOp::Set);
-	SET_DWORD_STAT(STAT_ChaosCounter_NumDirtyAABBTreeElements, DirtyElements);
+	Chaos::AABBTreeStatistics TreeStats;
+	GetAABBTreeStats(GetSpacialAcceleration()->AsChecked<SpatialAccelerationCollection>(), TreeStats);
+
+	CSV_CUSTOM_STAT(ChaosPhysics, AABBTreeDirtyElementCount, TreeStats.StatNumDirtyElements, ECsvCustomStatOp::Set);
+	SET_DWORD_STAT(STAT_ChaosCounter_NumDirtyAABBTreeElements, TreeStats.StatNumDirtyElements);
+
+	CSV_CUSTOM_STAT(ChaosPhysics, AABBTreeDirtyGridOverflowCount, TreeStats.StatNumGridOverflowElements, ECsvCustomStatOp::Set);
+	SET_DWORD_STAT(STAT_ChaosCounter_NumDirtyGridOverflowElements, TreeStats.StatNumGridOverflowElements);
+
+	CSV_CUSTOM_STAT(ChaosPhysics, AABBTreeDirtyElementTooLargeCount, TreeStats.StatNumElementsTooLargeForGrid, ECsvCustomStatOp::Set);
+	SET_DWORD_STAT(STAT_ChaosCounter_NumDirtyElementsTooLargeForGrid, TreeStats.StatNumElementsTooLargeForGrid);
+
+	CSV_CUSTOM_STAT(ChaosPhysics, AABBTreeDirtyElementNonEmptyCellCount, TreeStats.StatNumNonEmptyCellsInGrid, ECsvCustomStatOp::Set);
+	SET_DWORD_STAT(STAT_ChaosCounter_NumDirtyNonEmptyCellsInGrid, TreeStats.StatNumNonEmptyCellsInGrid);
 
 	check(IsCompletionEventComplete())
 	//check(PhysicsTickTask->IsComplete());
