Fixed access to inner physics materials during update. If an update is requested before we've made the physics engine material we would hit an ensure. This was harmless but it's due to not correctly accessing the material during the update - this fix sets it up correctly.

diff --git a/Engine/Source/Runtime/PhysicsCore/Private/PhysicalMaterial.cpp b/Engine/Source/Runtime/PhysicsCore/Private/PhysicalMaterial.cpp
--- a/Engine/Source/Runtime/PhysicsCore/Private/PhysicalMaterial.cpp
+++ b/Engine/Source/Runtime/PhysicsCore/Private/PhysicalMaterial.cpp
@@ -43,12 +43,20 @@ UPhysicalMaterial::~UPhysicalMaterial() = default;
 #if WITH_EDITOR
 void UPhysicalMaterial::PostEditChangeProperty(FPropertyChangedEvent& PropertyChangedEvent)
 {
+	bool bSkipUpdate = false;
 	if(!MaterialHandle)
 	{
-		MaterialHandle = MakeUnique<FPhysicsMaterialHandle>();
+		// If we don't currently have a material calling GetPhysicsMaterial will already call update as a side effect
+		// to set the initial state - so we can skip it in that case.
+		bSkipUpdate = true;
+	}
+
+	FPhysicsMaterialHandle& PhysMaterial = GetPhysicsMaterial();
+
+	if(!bSkipUpdate)
+	{
+		FChaosEngineInterface::UpdateMaterial(*MaterialHandle, this);
 	}
-	// Update PhysX material last so we have a valid Parent
-	FChaosEngineInterface::UpdateMaterial(*MaterialHandle, this);
 
 	Super::PostEditChangeProperty(PropertyChangedEvent);
 }
