LLM - finer chaos categories

diff --git a/Engine/Source/Runtime/Core/Private/HAL/LowLevelMemTracker.cpp b/Engine/Source/Runtime/Core/Private/HAL/LowLevelMemTracker.cpp
--- a/Engine/Source/Runtime/Core/Private/HAL/LowLevelMemTracker.cpp
+++ b/Engine/Source/Runtime/Core/Private/HAL/LowLevelMemTracker.cpp
@@ -158,6 +158,12 @@ DECLARE_LLM_MEMORY_STAT(TEXT("ChaosParticles"), STAT_ChaosParticlesLLM, STATGROU
 DECLARE_LLM_MEMORY_STAT(TEXT("ChaosLandscape"), STAT_ChaosLandscapeLLM, STATGROUP_LLMFULL);
 DECLARE_LLM_MEMORY_STAT(TEXT("ChaosTrimesh"), STAT_ChaosTrimeshLLM, STATGROUP_LLMFULL);
 DECLARE_LLM_MEMORY_STAT(TEXT("ChaosConvex"), STAT_ChaosConvexLLM, STATGROUP_LLMFULL);
+DECLARE_LLM_MEMORY_STAT(TEXT("ChaosScene"), STAT_ChaosSceneLLM, STATGROUP_LLMFULL);
+DECLARE_LLM_MEMORY_STAT(TEXT("ChaosUpdate"), STAT_ChaosUpdateLLM, STATGROUP_LLMFULL);
+DECLARE_LLM_MEMORY_STAT(TEXT("ChaosActor"), STAT_ChaosActorLLM, STATGROUP_LLMFULL);
+DECLARE_LLM_MEMORY_STAT(TEXT("ChaosBody"), STAT_ChaosBodyLLM, STATGROUP_LLMFULL);
+DECLARE_LLM_MEMORY_STAT(TEXT("ChaosConstraint"), STAT_ChaosConstraintLLM, STATGROUP_LLMFULL);
+DECLARE_LLM_MEMORY_STAT(TEXT("ChaosMaterial"), STAT_ChaosMaterialLLM, STATGROUP_LLMFULL);
 DECLARE_LLM_MEMORY_STAT(TEXT("EnginePreInit"), STAT_EnginePreInitLLM, STATGROUP_LLMFULL);
 DECLARE_LLM_MEMORY_STAT(TEXT("EngineInit"), STAT_EngineInitLLM, STATGROUP_LLMFULL);
 DECLARE_LLM_MEMORY_STAT(TEXT("Rendering Thread"), STAT_RenderingThreadLLM, STATGROUP_LLMFULL);
diff --git a/Engine/Source/Runtime/Core/Public/HAL/LowLevelMemTracker.h b/Engine/Source/Runtime/Core/Public/HAL/LowLevelMemTracker.h
--- a/Engine/Source/Runtime/Core/Public/HAL/LowLevelMemTracker.h
+++ b/Engine/Source/Runtime/Core/Public/HAL/LowLevelMemTracker.h
@@ -204,6 +204,12 @@ enum class ELLMTagSet : uint8
 	macro(ChaosLandscape,						"ChaosLandscape",				GET_STATFNAME(STAT_ChaosLandscapeLLM),						GET_STATFNAME(STAT_ChaosSummaryLLM),			ELLMTag::Physics)\
 	macro(ChaosTrimesh,							"ChaosTrimesh",					GET_STATFNAME(STAT_ChaosTrimeshLLM),						GET_STATFNAME(STAT_ChaosSummaryLLM),			ELLMTag::Physics)\
 	macro(ChaosConvex,							"ChaosConvex",					GET_STATFNAME(STAT_ChaosConvexLLM),							GET_STATFNAME(STAT_ChaosSummaryLLM),			ELLMTag::Physics)\
+	macro(ChaosScene,							"ChaosScene",					GET_STATFNAME(STAT_ChaosSceneLLM),							GET_STATFNAME(STAT_ChaosSummaryLLM),			ELLMTag::Physics)\
+	macro(ChaosUpdate,							"ChaosUpdate",					GET_STATFNAME(STAT_ChaosUpdateLLM),							GET_STATFNAME(STAT_ChaosSummaryLLM),			ELLMTag::Physics)\
+	macro(ChaosActor,							"ChaosActor",					GET_STATFNAME(STAT_ChaosActorLLM),							GET_STATFNAME(STAT_ChaosSummaryLLM),			ELLMTag::Physics)\
+	macro(ChaosBody,							"ChaosBody",					GET_STATFNAME(STAT_ChaosBodyLLM),							GET_STATFNAME(STAT_ChaosSummaryLLM),			ELLMTag::Physics)\
+	macro(ChaosConstraint,						"ChaosConstraint",				GET_STATFNAME(STAT_ChaosConstraintLLM),						GET_STATFNAME(STAT_ChaosSummaryLLM),			ELLMTag::Physics)\
+	macro(ChaosMaterial,						"ChaosMaterial",				GET_STATFNAME(STAT_ChaosMaterialLLM),						GET_STATFNAME(STAT_ChaosSummaryLLM),			ELLMTag::Physics)\
 	macro(EnginePreInitMemory,					"EnginePreInit",				GET_STATFNAME(STAT_EnginePreInitLLM),						GET_STATFNAME(STAT_EngineSummaryLLM),			-1)\
 	macro(EngineInitMemory,						"EngineInit",					GET_STATFNAME(STAT_EngineInitLLM),							GET_STATFNAME(STAT_EngineSummaryLLM),			-1)\
 	macro(RenderingThreadMemory,				"RenderingThread",				GET_STATFNAME(STAT_RenderingThreadLLM),						GET_STATFNAME(STAT_EngineSummaryLLM),			-1)\
@@ -242,7 +248,7 @@ enum class ELLMTag : LLM_TAG_TYPE
 
 	//------------------------------
 	// Platform tags
-	PlatformTagStart = 101,
+	PlatformTagStart = 107,
 	PlatformTagEnd = 149,
 
 	//------------------------------
diff --git a/Engine/Source/Runtime/Engine/Private/PhysicsEngine/BodyInstance.cpp b/Engine/Source/Runtime/Engine/Private/PhysicsEngine/BodyInstance.cpp
--- a/Engine/Source/Runtime/Engine/Private/PhysicsEngine/BodyInstance.cpp
+++ b/Engine/Source/Runtime/Engine/Private/PhysicsEngine/BodyInstance.cpp
@@ -1374,7 +1374,7 @@ bool FInitBodiesHelperBase::CreateShapesAndActors()
 void FInitBodiesHelperBase::InitBodies()
 {
 #if WITH_CHAOS
-	LLM_SCOPE(ELLMTag::Chaos);
+	LLM_SCOPE(ELLMTag::ChaosBody);
 #else
 	LLM_SCOPE(ELLMTag::PhysX);
 #endif
diff --git a/Engine/Source/Runtime/Engine/Private/PhysicsEngine/ConstraintInstance.cpp b/Engine/Source/Runtime/Engine/Private/PhysicsEngine/ConstraintInstance.cpp
--- a/Engine/Source/Runtime/Engine/Private/PhysicsEngine/ConstraintInstance.cpp
+++ b/Engine/Source/Runtime/Engine/Private/PhysicsEngine/ConstraintInstance.cpp
@@ -420,7 +420,7 @@ bool GetActorRefs(FBodyInstance* Body1, FBodyInstance* Body2, FPhysicsActorHandl
 bool FConstraintInstance::CreateJoint_AssumesLocked(const FPhysicsActorHandle& InActorRef1, const FPhysicsActorHandle& InActorRef2)
 {
 #if WITH_CHAOS
-	LLM_SCOPE(ELLMTag::Chaos);
+	LLM_SCOPE(ELLMTag::ChaosConstraint);
 #else
 	LLM_SCOPE(ELLMTag::PhysX);
 #endif
diff --git a/Engine/Source/Runtime/Engine/Private/PhysicsEngine/Experimental/PhysScene_Chaos.cpp b/Engine/Source/Runtime/Engine/Private/PhysicsEngine/Experimental/PhysScene_Chaos.cpp
--- a/Engine/Source/Runtime/Engine/Private/PhysicsEngine/Experimental/PhysScene_Chaos.cpp
+++ b/Engine/Source/Runtime/Engine/Private/PhysicsEngine/Experimental/PhysScene_Chaos.cpp
@@ -390,7 +390,7 @@ FPhysScene_Chaos::FPhysScene_Chaos(AActor* InSolverActor
 #endif
 {
 #if WITH_CHAOS
-	LLM_SCOPE(ELLMTag::Chaos);
+	LLM_SCOPE(ELLMTag::ChaosScene);
 
 	PhysicsProxyToComponentMap.Reset();
 	ComponentToPhysicsProxyMap.Reset();
diff --git a/Engine/Source/Runtime/Experimental/Chaos/Private/Chaos/Framework/PhysicsSolverBase.cpp b/Engine/Source/Runtime/Experimental/Chaos/Private/Chaos/Framework/PhysicsSolverBase.cpp
--- a/Engine/Source/Runtime/Experimental/Chaos/Private/Chaos/Framework/PhysicsSolverBase.cpp
+++ b/Engine/Source/Runtime/Experimental/Chaos/Private/Chaos/Framework/PhysicsSolverBase.cpp
@@ -98,7 +98,7 @@ namespace Chaos
 	{
 		using namespace Chaos;
 
-		LLM_SCOPE(ELLMTag::Chaos);
+		LLM_SCOPE(ELLMTag::ChaosUpdate);
 		SCOPE_CYCLE_COUNTER(STAT_ChaosTick);
 		CSV_SCOPED_TIMING_STAT_EXCLUSIVE(Physics);
 		CSV_SCOPED_TIMING_STAT(PhysicsVerbose, StepSolver);
diff --git a/Engine/Source/Runtime/Experimental/Chaos/Private/Chaos/PBDRigidsEvolution.cpp b/Engine/Source/Runtime/Experimental/Chaos/Private/Chaos/PBDRigidsEvolution.cpp
--- a/Engine/Source/Runtime/Experimental/Chaos/Private/Chaos/PBDRigidsEvolution.cpp
+++ b/Engine/Source/Runtime/Experimental/Chaos/Private/Chaos/PBDRigidsEvolution.cpp
@@ -343,8 +343,6 @@ namespace Chaos
 			return Result;
 		};
 
-		LLM_SCOPE(ELLMTag::ChaosAcceleration);
-
 		uint8 ActiveBucketsMask = SpatialCollectionFactory.GetActiveBucketsMask();
 		TArray<TSOAView<FSpatialAccelerationCache>> ViewsPerBucket[FSpatialAccelerationIdx::MaxBuckets];
 		TArray<FSpatialAccelerationIdx> TimeSlicedBucketsToCreate;
diff --git a/Engine/Source/Runtime/Experimental/Chaos/Private/Chaos/PhysicalMaterials.cpp b/Engine/Source/Runtime/Experimental/Chaos/Private/Chaos/PhysicalMaterials.cpp
--- a/Engine/Source/Runtime/Experimental/Chaos/Private/Chaos/PhysicalMaterials.cpp
+++ b/Engine/Source/Runtime/Experimental/Chaos/Private/Chaos/PhysicalMaterials.cpp
@@ -100,7 +100,7 @@ namespace Chaos
 
 	FMaterialHandle FPhysicalMaterialManager::Create()
 	{
-		LLM_SCOPE(ELLMTag::Chaos);
+		LLM_SCOPE(ELLMTag::ChaosMaterial);
 
 		check(IsInGameThread());
 		FMaterialHandle OutHandle;
@@ -124,7 +124,7 @@ namespace Chaos
 
 	void FPhysicalMaterialManager::Destroy(FMaterialHandle InHandle)
 	{
-		LLM_SCOPE(ELLMTag::Chaos);
+		LLM_SCOPE(ELLMTag::ChaosMaterial);
 
 		check(IsInGameThread());
 		if(InHandle.InnerHandle.IsValid())
diff --git a/Engine/Source/Runtime/Experimental/Chaos/Private/PBDRigidsSolver.cpp b/Engine/Source/Runtime/Experimental/Chaos/Private/PBDRigidsSolver.cpp
--- a/Engine/Source/Runtime/Experimental/Chaos/Private/PBDRigidsSolver.cpp
+++ b/Engine/Source/Runtime/Experimental/Chaos/Private/PBDRigidsSolver.cpp
@@ -262,7 +262,7 @@ namespace Chaos
 
 		void DoWork()
 		{
-			LLM_SCOPE(ELLMTag::Chaos);
+			LLM_SCOPE(ELLMTag::ChaosUpdate);
 			UE_LOG(LogPBDRigidsSolver, Verbose, TEXT("AdvanceOneTimeStepTask::DoWork()"));
 			MSolver->StartingSceneSimulation();
 
@@ -467,7 +467,7 @@ namespace Chaos
 
 	void FPBDRigidsSolver::RegisterObject(FSingleParticlePhysicsProxy* Proxy)
 	{
-		LLM_SCOPE(ELLMTag::Chaos);
+		LLM_SCOPE(ELLMTag::ChaosBody);
 
 		UE_LOG(LogPBDRigidsSolver, Verbose, TEXT("FPBDRigidsSolver::RegisterObject()"));
 		auto& RigidBody_External = Proxy->GetGameThreadAPI();
@@ -568,6 +568,7 @@ namespace Chaos
 
 	void FPBDRigidsSolver::RegisterObject(FGeometryCollectionPhysicsProxy* InProxy)
 	{
+		LLM_SCOPE(ELLMTag::ChaosBody);
 		UE_LOG(LogPBDRigidsSolver, Verbose, TEXT("FPBDRigidsSolver::RegisterObject(FGeometryCollectionPhysicsProxy*)"));
 		InProxy->SetSolver(this);
 		InProxy->Initialize(GetEvolution());
@@ -608,6 +609,7 @@ namespace Chaos
 
 	void FPBDRigidsSolver::RegisterObject(Chaos::FJointConstraint* GTConstraint)
 	{
+		LLM_SCOPE(ELLMTag::ChaosConstraint);
 		FJointConstraintPhysicsProxy* JointProxy = new FJointConstraintPhysicsProxy(GTConstraint, nullptr);
 		JointProxy->SetSolver(this);
 
@@ -650,6 +652,7 @@ namespace Chaos
 
 	void FPBDRigidsSolver::RegisterObject(Chaos::FSuspensionConstraint* GTConstraint)
 	{
+		LLM_SCOPE(ELLMTag::ChaosConstraint);
 		FSuspensionConstraintPhysicsProxy* SuspensionProxy = new FSuspensionConstraintPhysicsProxy(GTConstraint, nullptr);
 		SuspensionProxy->SetSolver(this);
 
diff --git a/Engine/Source/Runtime/PhysicsCore/Private/ChaosEngineInterface.cpp b/Engine/Source/Runtime/PhysicsCore/Private/ChaosEngineInterface.cpp
--- a/Engine/Source/Runtime/PhysicsCore/Private/ChaosEngineInterface.cpp
+++ b/Engine/Source/Runtime/PhysicsCore/Private/ChaosEngineInterface.cpp
@@ -148,7 +148,6 @@ const FQuat& FPhysicsShapeAdapter_Chaos::GetGeomOrientation() const
 
 void FChaosEngineInterface::AddActorToSolver(FPhysicsActorHandle& Handle,Chaos::FPhysicsSolver* Solver)
 {
-	LLM_SCOPE(ELLMTag::Chaos);
 	Solver->RegisterObject(Handle);
 }
 
@@ -959,7 +958,7 @@ FPhysicsConstraintHandle FChaosEngineInterface::CreateConstraint(const FPhysicsA
 		{
 			if(InActorRef1->GetSolverBase() && InActorRef2->GetSolverBase())
 			{
-				LLM_SCOPE(ELLMTag::Chaos);
+				LLM_SCOPE(ELLMTag::ChaosConstraint);
 
 				auto* JointConstraint = new Chaos::FJointConstraint();
 				ConstraintRef.Constraint = JointConstraint;
@@ -974,7 +973,7 @@ FPhysicsConstraintHandle FChaosEngineInterface::CreateConstraint(const FPhysicsA
 		}
 		else if (InActorRef1 != nullptr || InActorRef2 != nullptr)
 		{
-			LLM_SCOPE(ELLMTag::Chaos);
+			LLM_SCOPE(ELLMTag::ChaosConstraint);
 
 			FPhysicsActorHandle ValidParticle = InActorRef1;
 			bool bSwapped = false;
@@ -1045,7 +1044,7 @@ FPhysicsConstraintHandle FChaosEngineInterface::CreateSuspension(const FPhysicsA
 		{
 			if (InActorRef->GetSolverBase())
 			{
-				LLM_SCOPE(ELLMTag::Chaos);
+				LLM_SCOPE(ELLMTag::ChaosConstraint);
 
 				auto* SuspensionConstraint = new Chaos::FSuspensionConstraint();
 				ConstraintRef.Constraint = SuspensionConstraint;
@@ -1078,7 +1077,6 @@ void FChaosEngineInterface::ReleaseConstraint(FPhysicsConstraintHandle& InConstr
 	using namespace Chaos;
 	if (bEnableChaosJointConstraints)
 	{
-		LLM_SCOPE(ELLMTag::Chaos);
 		if (InConstraintRef.IsValid() && InConstraintRef.Constraint->IsType(Chaos::EConstraintType::JointConstraintType))
 		{
 			if (FJointConstraint* Constraint = static_cast<FJointConstraint*>(InConstraintRef.Constraint))
@@ -1614,7 +1612,7 @@ int32 FChaosEngineInterface::GetAllShapes_AssumedLocked(const FPhysicsActorHandl
 
 void FChaosEngineInterface::CreateActor(const FActorCreationParams& InParams,FPhysicsActorHandle& Handle)
 {
-	LLM_SCOPE(ELLMTag::Chaos);
+	LLM_SCOPE(ELLMTag::ChaosActor);
 	using namespace Chaos;
 
 	TUniquePtr<FGeometryParticle> Particle;
@@ -1662,7 +1660,6 @@ void FChaosEngineInterface::CreateActor(const FActorCreationParams& InParams,FPh
 
 void FChaosEngineInterface::ReleaseActor(FPhysicsActorHandle& Handle,FChaosScene* InScene,bool bNeverDerferRelease)
 {
-	LLM_SCOPE(ELLMTag::Chaos);
 	if(!Handle)
 	{
 		UE_LOG(LogChaos,Warning,TEXT("Attempting to release an actor with a null handle"));
diff --git a/Engine/Source/Runtime/PhysicsCore/Private/ChaosScene.cpp b/Engine/Source/Runtime/PhysicsCore/Private/ChaosScene.cpp
--- a/Engine/Source/Runtime/PhysicsCore/Private/ChaosScene.cpp
+++ b/Engine/Source/Runtime/PhysicsCore/Private/ChaosScene.cpp
@@ -61,7 +61,7 @@ FChaosScene::FChaosScene(
 	, SceneSolver(nullptr)
 	, Owner(OwnerPtr)
 {
-	LLM_SCOPE(ELLMTag::Chaos);
+	LLM_SCOPE(ELLMTag::ChaosScene);
 
 	ChaosModule = FChaosSolversModule::GetModule();
 	check(ChaosModule);
