Type truncation fixes - Physics

diff --git a/Engine/Source/Runtime/PhysicsCore/Private/ChaosInterfaceWrapperCore.cpp b/Engine/Source/Runtime/PhysicsCore/Private/ChaosInterfaceWrapperCore.cpp
--- a/Engine/Source/Runtime/PhysicsCore/Private/ChaosInterfaceWrapperCore.cpp
+++ b/Engine/Source/Runtime/PhysicsCore/Private/ChaosInterfaceWrapperCore.cpp
@@ -40,12 +40,12 @@ namespace ChaosInterface
 		return ImplicitTypeToCollisionType(ImplicitObjectType);
 	}
 
-	float GetRadius(const Chaos::FCapsule& InCapsule)
+	Chaos::FReal GetRadius(const Chaos::FCapsule& InCapsule)
 	{
 		return InCapsule.GetRadius();
 	}
 
-	float GetHalfHeight(const Chaos::FCapsule& InCapsule)
+	Chaos::FReal GetHalfHeight(const Chaos::FCapsule& InCapsule)
 	{
 		return InCapsule.GetHeight() / 2.;
 	}
diff --git a/Engine/Source/Runtime/PhysicsCore/Private/SQAccelerator.cpp b/Engine/Source/Runtime/PhysicsCore/Private/SQAccelerator.cpp
--- a/Engine/Source/Runtime/PhysicsCore/Private/SQAccelerator.cpp
+++ b/Engine/Source/Runtime/PhysicsCore/Private/SQAccelerator.cpp
@@ -69,7 +69,7 @@ struct FPreFilterInfo
 	int32 ActorIdx;
 };
 
-void FillHitHelper(ChaosInterface::FLocationHit& Hit, const float Distance, const FVector& WorldPosition, const FVector& WorldNormal, int32 FaceIdx, bool bComputeMTD)
+void FillHitHelper(ChaosInterface::FLocationHit& Hit, const Chaos::FReal Distance, const FVector& WorldPosition, const FVector& WorldNormal, int32 FaceIdx, bool bComputeMTD)
 {
 	Hit.Distance = Distance;
 	Hit.WorldPosition = WorldPosition;
@@ -78,7 +78,7 @@ void FillHitHelper(ChaosInterface::FLocationHit& Hit, const float Distance, cons
 	Hit.FaceIndex = FaceIdx;
 }
 
-void FillHitHelper(ChaosInterface::FOverlapHit& Hit, const float Distance, const FVector& WorldPosition, const FVector& WorldNormal, int32 FaceIdx, bool bComputeMTD)
+void FillHitHelper(ChaosInterface::FOverlapHit& Hit, const Chaos::FReal Distance, const FVector& WorldPosition, const FVector& WorldNormal, int32 FaceIdx, bool bComputeMTD)
 {
 }
 
diff --git a/Engine/Source/Runtime/PhysicsCore/Public/ChaosInterfaceWrapperCore.h b/Engine/Source/Runtime/PhysicsCore/Public/ChaosInterfaceWrapperCore.h
--- a/Engine/Source/Runtime/PhysicsCore/Public/ChaosInterfaceWrapperCore.h
+++ b/Engine/Source/Runtime/PhysicsCore/Public/ChaosInterfaceWrapperCore.h
@@ -106,9 +106,9 @@ FORCEINLINE ECollisionShapeType GetType(const Chaos::FImplicitObject& InGeometry
 	return GetImplicitType(InGeometry);
 }
 
-PHYSICSCORE_API float GetRadius(const Chaos::FCapsule& InCapsule);
+PHYSICSCORE_API Chaos::FReal GetRadius(const Chaos::FCapsule& InCapsule);
 
-PHYSICSCORE_API float GetHalfHeight(const Chaos::FCapsule& InCapsule);
+PHYSICSCORE_API Chaos::FReal GetHalfHeight(const Chaos::FCapsule& InCapsule);
 
 
 inline bool HadInitialOverlap(const FLocationHit& Hit)
@@ -126,7 +126,7 @@ inline Chaos::FGeometryParticle* GetActor(const FActorShape& Hit)
 	return Hit.Actor;
 }
 
-inline float GetDistance(const FLocationHit& Hit)
+inline Chaos::FReal GetDistance(const FLocationHit& Hit)
 {
 	return Hit.Distance;
 }
diff --git a/Engine/Source/Runtime/PhysicsCore/Public/ChaosSQTypes.h b/Engine/Source/Runtime/PhysicsCore/Public/ChaosSQTypes.h
--- a/Engine/Source/Runtime/PhysicsCore/Public/ChaosSQTypes.h
+++ b/Engine/Source/Runtime/PhysicsCore/Public/ChaosSQTypes.h
@@ -53,7 +53,7 @@ namespace ChaosInterface
 		FHitFlags Flags;
 		FVector WorldPosition;
 		FVector WorldNormal;
-		float Distance;
+		FVector::FReal Distance;
 
 		void Serialize(Chaos::FChaosArchive& Ar);
 
diff --git a/Engine/Source/Runtime/PhysicsCore/Public/CollisionShape.h b/Engine/Source/Runtime/PhysicsCore/Public/CollisionShape.h
--- a/Engine/Source/Runtime/PhysicsCore/Public/CollisionShape.h
+++ b/Engine/Source/Runtime/PhysicsCore/Public/CollisionShape.h
@@ -78,7 +78,7 @@ struct FCollisionShape
 	}
 
 	/** Utility function to Set Box and dimension */
-	void SetBox(const FVector& HalfExtent)
+	void SetBox(const FVector3f& HalfExtent)
 	{
 		ShapeType = ECollisionShape::Box;
 		Box.HalfExtentX = HalfExtent.X;
@@ -102,7 +102,7 @@ struct FCollisionShape
 	}
 
 	/** Utility function to set Capsule from Extent data */
-	void SetCapsule(const FVector& Extent)
+	void SetCapsule(const FVector3f& Extent)
 	{
 		ShapeType = ECollisionShape::Capsule;
 		Capsule.Radius = FMath::Max(Extent.X, Extent.Y);
diff --git a/Engine/Source/Runtime/PhysicsCore/Public/PhysicsInterfaceTypesCore.h b/Engine/Source/Runtime/PhysicsCore/Public/PhysicsInterfaceTypesCore.h
--- a/Engine/Source/Runtime/PhysicsCore/Public/PhysicsInterfaceTypesCore.h
+++ b/Engine/Source/Runtime/PhysicsCore/Public/PhysicsInterfaceTypesCore.h
@@ -90,7 +90,7 @@ struct FBodyCollisionData
 	FBodyCollisionFlags CollisionFlags;
 };
 
-static void SetupNonUniformHelper(FVector InScale3D, float& OutMinScale, float& OutMinScaleAbs, FVector& OutScale3DAbs)
+static void SetupNonUniformHelper(FVector InScale3D, double& OutMinScale, double& OutMinScaleAbs, FVector& OutScale3DAbs)
 {
 	// if almost zero, set min scale
 	// @todo fixme
@@ -113,6 +113,15 @@ static void SetupNonUniformHelper(FVector InScale3D, float& OutMinScale, float&
 	}
 }
 
+static void SetupNonUniformHelper(FVector InScale3D, float& OutMinScale, float& OutMinScaleAbs, FVector& OutScale3DAbs)
+{
+	double OutMinScaleD, OutMinScaleAbsD;
+	SetupNonUniformHelper(InScale3D, OutMinScaleD, OutMinScaleAbsD, OutScale3DAbs);
+	OutMinScale = static_cast<float>(OutMinScaleD);	// LWC_TODO: Precision loss?
+	OutMinScaleAbs = static_cast<float>(OutMinScaleAbsD);
+}
+
+
 /** Util to determine whether to use NegX version of mesh, and what transform (rotation) to apply. */
 static bool CalcMeshNegScaleCompensation(const FVector& InScale3D, FTransform& OutTransform)
 {
@@ -349,7 +358,7 @@ typedef ChaosFlags<FChaosQueryFlag::Enum, uint16> FChaosQueryFlags;
 
 inline FChaosQueryFlags U2CQueryFlags(FQueryFlags Flags)
 {
-	uint32 Result = 0;
+	uint16 Result = 0;
 	if (Flags & EQueryFlags::PreFilter)
 	{
 		Result |= FChaosQueryFlag::ePREFILTER;
