Temporarily revert type truncation fix that leads to collision failures until it can be investigated.

diff --git a/Engine/Source/Runtime/PhysicsCore/Public/ChaosSQTypes.h b/Engine/Source/Runtime/PhysicsCore/Public/ChaosSQTypes.h
--- a/Engine/Source/Runtime/PhysicsCore/Public/ChaosSQTypes.h
+++ b/Engine/Source/Runtime/PhysicsCore/Public/ChaosSQTypes.h
@@ -53,7 +53,7 @@ namespace ChaosInterface
 		FHitFlags Flags;
 		FVector WorldPosition;
 		FVector WorldNormal;
-		FVector::FReal Distance;
+		float Distance;	// LWC_TODO: Should be FVector::FReal, but that causes precision issues resulting in collision failures. Investigate!
 
 		void Serialize(Chaos::FChaosArchive& Ar);
 
