Simple chaos optimization
- Disable the csv category AABBTreeExpensiveStats by default
- Disable stat recording altogether in shipping
This saves 0.4ms on the GT when the CSV profiler is compiled in

diff --git a/Engine/Source/Runtime/PhysicsCore/Private/ChaosScene.cpp b/Engine/Source/Runtime/PhysicsCore/Private/ChaosScene.cpp
--- a/Engine/Source/Runtime/PhysicsCore/Private/ChaosScene.cpp
+++ b/Engine/Source/Runtime/PhysicsCore/Private/ChaosScene.cpp
@@ -38,7 +38,7 @@
 
 DECLARE_CYCLE_STAT(TEXT("Update Kinematics On Deferred SkelMeshes"),STAT_UpdateKinematicsOnDeferredSkelMeshesChaos,STATGROUP_Physics);
 CSV_DEFINE_CATEGORY(ChaosPhysics,true);
-CSV_DEFINE_CATEGORY(AABBTreeExpensiveStats, true);
+CSV_DEFINE_CATEGORY(AABBTreeExpensiveStats, false);
 
 // Stat Counters
 DECLARE_DWORD_ACCUMULATOR_STAT(TEXT("NumDirtyAABBTreeElements"), STAT_ChaosCounter_NumDirtyAABBTreeElements, STATGROUP_ChaosCounters);
@@ -444,6 +444,7 @@ void FChaosScene::SyncBodies(TSolver* Solver)
 // Accumulate all the AABBTree stats
 void GetAABBTreeStats(Chaos::ISpatialAccelerationCollection<Chaos::FAccelerationStructureHandle, Chaos::FReal, 3>& Collection, Chaos::AABBTreeStatistics& OutAABBTreeStatistics, Chaos::AABBTreeExpensiveStatistics& OutAABBTreeExpensiveStatistics)
 {
+	CSV_SCOPED_TIMING_STAT(AABBTreeExpensiveStats, GetAABBTreeStats);
 	using namespace Chaos;
 	OutAABBTreeStatistics.Reset();
 	TArray<FSpatialAccelerationIdx> SpatialIndices = Collection.GetAllSpatialIndices();
@@ -480,32 +481,36 @@ void FChaosScene::EndFrame()
 		return;
 	}
 
-	Chaos::AABBTreeStatistics TreeStats;
-	Chaos::AABBTreeExpensiveStatistics TreeExpensiveStats;
-	GetAABBTreeStats(GetSpacialAcceleration()->AsChecked<SpatialAccelerationCollection>(), TreeStats, TreeExpensiveStats);
+#if !UE_BUILD_SHIPPING
+	{
+		Chaos::AABBTreeStatistics TreeStats;
+		Chaos::AABBTreeExpensiveStatistics TreeExpensiveStats;
+		GetAABBTreeStats(GetSpacialAcceleration()->AsChecked<SpatialAccelerationCollection>(), TreeStats, TreeExpensiveStats);
 
-	CSV_CUSTOM_STAT(ChaosPhysics, AABBTreeDirtyElementCount, TreeStats.StatNumDirtyElements, ECsvCustomStatOp::Set);
-	SET_DWORD_STAT(STAT_ChaosCounter_NumDirtyAABBTreeElements, TreeStats.StatNumDirtyElements);
+		CSV_CUSTOM_STAT(ChaosPhysics, AABBTreeDirtyElementCount, TreeStats.StatNumDirtyElements, ECsvCustomStatOp::Set);
+		SET_DWORD_STAT(STAT_ChaosCounter_NumDirtyAABBTreeElements, TreeStats.StatNumDirtyElements);
 
-	CSV_CUSTOM_STAT(ChaosPhysics, AABBTreeDirtyGridOverflowCount, TreeStats.StatNumGridOverflowElements, ECsvCustomStatOp::Set);
-	SET_DWORD_STAT(STAT_ChaosCounter_NumDirtyGridOverflowElements, TreeStats.StatNumGridOverflowElements);
+		CSV_CUSTOM_STAT(ChaosPhysics, AABBTreeDirtyGridOverflowCount, TreeStats.StatNumGridOverflowElements, ECsvCustomStatOp::Set);
+		SET_DWORD_STAT(STAT_ChaosCounter_NumDirtyGridOverflowElements, TreeStats.StatNumGridOverflowElements);
 
-	CSV_CUSTOM_STAT(ChaosPhysics, AABBTreeDirtyElementTooLargeCount, TreeStats.StatNumElementsTooLargeForGrid, ECsvCustomStatOp::Set);
-	SET_DWORD_STAT(STAT_ChaosCounter_NumDirtyElementsTooLargeForGrid, TreeStats.StatNumElementsTooLargeForGrid);
+		CSV_CUSTOM_STAT(ChaosPhysics, AABBTreeDirtyElementTooLargeCount, TreeStats.StatNumElementsTooLargeForGrid, ECsvCustomStatOp::Set);
+		SET_DWORD_STAT(STAT_ChaosCounter_NumDirtyElementsTooLargeForGrid, TreeStats.StatNumElementsTooLargeForGrid);
 
-	CSV_CUSTOM_STAT(ChaosPhysics, AABBTreeDirtyElementNonEmptyCellCount, TreeStats.StatNumNonEmptyCellsInGrid, ECsvCustomStatOp::Set);
-	SET_DWORD_STAT(STAT_ChaosCounter_NumDirtyNonEmptyCellsInGrid, TreeStats.StatNumNonEmptyCellsInGrid);
+		CSV_CUSTOM_STAT(ChaosPhysics, AABBTreeDirtyElementNonEmptyCellCount, TreeStats.StatNumNonEmptyCellsInGrid, ECsvCustomStatOp::Set);
+		SET_DWORD_STAT(STAT_ChaosCounter_NumDirtyNonEmptyCellsInGrid, TreeStats.StatNumNonEmptyCellsInGrid);
 
 #if CSV_PROFILER
-	if (FCsvProfiler::Get()->IsCapturing() && FCsvProfiler::Get()->IsCategoryEnabled(CSV_CATEGORY_INDEX(AABBTreeExpensiveStats)))
-	{
-		CSV_CUSTOM_STAT(AABBTreeExpensiveStats, AABBTreeMaxNumLeaves, TreeExpensiveStats.StatMaxNumLeaves, ECsvCustomStatOp::Set);
-		CSV_CUSTOM_STAT(AABBTreeExpensiveStats, AABBTreeMaxDirtyElements, TreeExpensiveStats.StatMaxDirtyElements, ECsvCustomStatOp::Set);
-		CSV_CUSTOM_STAT(AABBTreeExpensiveStats, AABBTreeMaxTreeDepth, TreeExpensiveStats.StatMaxTreeDepth, ECsvCustomStatOp::Set);
-		CSV_CUSTOM_STAT(AABBTreeExpensiveStats, AABBTreeMaxLeafSize, TreeExpensiveStats.StatMaxLeafSize, ECsvCustomStatOp::Set);
-		CSV_CUSTOM_STAT(AABBTreeExpensiveStats, AABBTreeGlobalPayloadsSize, TreeExpensiveStats.StatGlobalPayloadsSize, ECsvCustomStatOp::Set);
+		if (FCsvProfiler::Get()->IsCapturing() && FCsvProfiler::Get()->IsCategoryEnabled(CSV_CATEGORY_INDEX(AABBTreeExpensiveStats)))
+		{
+			CSV_CUSTOM_STAT(AABBTreeExpensiveStats, AABBTreeMaxNumLeaves, TreeExpensiveStats.StatMaxNumLeaves, ECsvCustomStatOp::Set);
+			CSV_CUSTOM_STAT(AABBTreeExpensiveStats, AABBTreeMaxDirtyElements, TreeExpensiveStats.StatMaxDirtyElements, ECsvCustomStatOp::Set);
+			CSV_CUSTOM_STAT(AABBTreeExpensiveStats, AABBTreeMaxTreeDepth, TreeExpensiveStats.StatMaxTreeDepth, ECsvCustomStatOp::Set);
+			CSV_CUSTOM_STAT(AABBTreeExpensiveStats, AABBTreeMaxLeafSize, TreeExpensiveStats.StatMaxLeafSize, ECsvCustomStatOp::Set);
+			CSV_CUSTOM_STAT(AABBTreeExpensiveStats, AABBTreeGlobalPayloadsSize, TreeExpensiveStats.StatGlobalPayloadsSize, ECsvCustomStatOp::Set);
+		}
+#endif // CSV_PROFILER
 	}
-#endif
+#endif // UE_BUILD_SHIPPING
 
 	check(IsCompletionEventComplete())
 	//check(PhysicsTickTask->IsComplete());
