More template boilerplate to be able to do scene queries against PT data

diff --git a/Engine/Source/Runtime/Engine/Private/Collision/CollisionConversions.cpp b/Engine/Source/Runtime/Engine/Private/Collision/CollisionConversions.cpp
--- a/Engine/Source/Runtime/Engine/Private/Collision/CollisionConversions.cpp
+++ b/Engine/Source/Runtime/Engine/Private/Collision/CollisionConversions.cpp
@@ -508,6 +508,8 @@ template EConvertQueryResult ConvertTraceResults<FHitSweep>(bool& OutHasValidBlo
 template EConvertQueryResult ConvertTraceResults<FHitSweep>(bool& OutHasValidBlockingHit, const UWorld* World, int32 NumHits, FHitSweep* Hits, float CheckLength, const FCollisionFilterData& QueryFilter, FHitResult& OutHit, const FVector& StartLoc, const FVector& EndLoc, const FPhysicsGeometry& Geom, const FTransform& QueryTM, float MaxDistance, bool bReturnFaceIndex, bool bReturnPhysMat);
 template EConvertQueryResult ConvertTraceResults<FHitRaycast>(bool& OutHasValidBlockingHit, const UWorld* World, int32 NumHits, FHitRaycast* Hits, float CheckLength, const FCollisionFilterData& QueryFilter, TArray<FHitResult>& OutHits, const FVector& StartLoc, const FVector& EndLoc, const FPhysicsGeometry& Geom, const FTransform& QueryTM, float MaxDistance, bool bReturnFaceIndex, bool bReturnPhysMat);
 template EConvertQueryResult ConvertTraceResults<FHitRaycast>(bool& OutHasValidBlockingHit, const UWorld* World, int32 NumHits, FHitRaycast* Hits, float CheckLength, const FCollisionFilterData& QueryFilter, FHitResult& OutHit, const FVector& StartLoc, const FVector& EndLoc, const FPhysicsGeometry& Geom, const FTransform& QueryTM, float MaxDistance, bool bReturnFaceIndex, bool bReturnPhysMat);
+template EConvertQueryResult ConvertTraceResults<FPTRaycastHit>(bool& OutHasValidBlockingHit, const UWorld* World, int32 NumHits, FPTRaycastHit* Hits, float CheckLength, const FCollisionFilterData& QueryFilter, TArray<FHitResult>& OutHits, const FVector& StartLoc, const FVector& EndLoc, const FPhysicsGeometry& Geom, const FTransform& QueryTM, float MaxDistance, bool bReturnFaceIndex, bool bReturnPhysMat);
+template EConvertQueryResult ConvertTraceResults<FPTRaycastHit>(bool& OutHasValidBlockingHit, const UWorld* World, int32 NumHits, FPTRaycastHit* Hits, float CheckLength, const FCollisionFilterData& QueryFilter, FHitResult& OutHit, const FVector& StartLoc, const FVector& EndLoc, const FPhysicsGeometry& Geom, const FTransform& QueryTM, float MaxDistance, bool bReturnFaceIndex, bool bReturnPhysMat);
 
 /** Util to convert an overlapped shape into a sweep hit result, returns whether it was a blocking hit. */
 template <typename THitLocation>
diff --git a/Engine/Source/Runtime/Engine/Private/Collision/SceneQuery.cpp b/Engine/Source/Runtime/Engine/Private/Collision/SceneQuery.cpp
--- a/Engine/Source/Runtime/Engine/Private/Collision/SceneQuery.cpp
+++ b/Engine/Source/Runtime/Engine/Private/Collision/SceneQuery.cpp
@@ -127,7 +127,7 @@ struct TSQTraits
 	static const ESweepOrRay GeometryQuery = InGeometryQuery;
 	using THitType = InHitType;
 	using TOutHits = typename TChooseClass<InSingleMultiOrTest == ESingleMultiOrTest::Multi, TArray<FHitResult>, FHitResult>::Result;
-	using THitBuffer = typename TChooseClass<InSingleMultiOrTest == ESingleMultiOrTest::Multi, FDynamicHitBuffer<InHitType>, typename TChooseClass<InGeometryQuery == ESweepOrRay::Sweep, FSingleHitBuffer<FHitSweep>, FSingleHitBuffer<FHitRaycast>>::Result >::Result;
+	using THitBuffer = typename TChooseClass<InSingleMultiOrTest == ESingleMultiOrTest::Multi, FDynamicHitBuffer<InHitType>, FSingleHitBuffer<InHitType>>::Result;
 
 	// GetNumHits - multi
 	template <ESingleMultiOrTest T = SingleMultiOrTest>
diff --git a/Engine/Source/Runtime/Engine/Private/Collision/SceneQueryLowLevel.cpp b/Engine/Source/Runtime/Engine/Private/Collision/SceneQueryLowLevel.cpp
--- a/Engine/Source/Runtime/Engine/Private/Collision/SceneQueryLowLevel.cpp
+++ b/Engine/Source/Runtime/Engine/Private/Collision/SceneQueryLowLevel.cpp
@@ -183,8 +183,8 @@ namespace
 }
 #endif
 
-
-void LowLevelRaycast(FPhysScene& Scene, const FVector& Start, const FVector& Dir, float DeltaMag, FPhysicsHitCallback<FHitRaycast>& HitBuffer, EHitFlags OutputFlags, FQueryFlags QueryFlags, const FCollisionFilterData& Filter, const FQueryFilterData& QueryFilterData, ICollisionQueryFilterCallbackBase* QueryCallback, const FQueryDebugParams& DebugParams)
+template <typename THitRaycast>
+void LowLevelRaycastImp(FPhysScene& Scene, const FVector& Start, const FVector& Dir, float DeltaMag, FPhysicsHitCallback<THitRaycast>& HitBuffer, EHitFlags OutputFlags, FQueryFlags QueryFlags, const FCollisionFilterData& Filter, const FQueryFilterData& QueryFilterData, ICollisionQueryFilterCallbackBase* QueryCallback, const FQueryDebugParams& DebugParams)
 {
 #if !defined(PHYSICS_INTERFACE_PHYSX) || !PHYSICS_INTERFACE_PHYSX
 	if (const auto& SolverAccelerationStructure = Scene.GetSpacialAcceleration())
@@ -196,9 +196,12 @@ void LowLevelRaycast(FPhysScene& Scene, const FVector& Start, const FVector& Dir
 			SQAccelerator.Raycast(Start, Dir, DeltaMag, HitBuffer, OutputFlags, QueryFilterData, *QueryCallback, DebugParams);
 		}
 
-		if (!!SerializeSQs && !!EnableRaycastSQCapture)
+		if constexpr(std::is_same<THitRaycast, FHitRaycast>::value)
 		{
-			RaycastSQCaptureHelper(Time, SQAccelerator, Scene, Start, Dir, DeltaMag, HitBuffer, OutputFlags, QueryFlags, Filter, QueryFilterData, QueryCallback, DebugParams);
+			if (!!SerializeSQs && !!EnableRaycastSQCapture)
+			{
+				RaycastSQCaptureHelper(Time, SQAccelerator, Scene, Start, Dir, DeltaMag, HitBuffer, OutputFlags, QueryFlags, Filter, QueryFilterData, QueryCallback, DebugParams);
+			}
 		}
 	}
 #else
@@ -220,6 +223,16 @@ void LowLevelRaycast(FPhysScene& Scene, const FVector& Start, const FVector& Dir
 #endif
 }
 
+void LowLevelRaycast(FPhysScene& Scene, const FVector& Start, const FVector& Dir, float DeltaMag, FPhysicsHitCallback<FHitRaycast>& HitBuffer, EHitFlags OutputFlags, FQueryFlags QueryFlags, const FCollisionFilterData& Filter, const FQueryFilterData& QueryFilterData, ICollisionQueryFilterCallbackBase* QueryCallback, const FQueryDebugParams& DebugParams)
+{
+	LowLevelRaycastImp(Scene, Start, Dir, DeltaMag, HitBuffer, OutputFlags, QueryFlags, Filter, QueryFilterData, QueryCallback, DebugParams);
+}
+
+void LowLevelRaycast(FPhysScene& Scene, const FVector& Start, const FVector& Dir, float DeltaMag, FPhysicsHitCallback<FPTRaycastHit>& HitBuffer, EHitFlags OutputFlags, FQueryFlags QueryFlags, const FCollisionFilterData& Filter, const FQueryFilterData& QueryFilterData, ICollisionQueryFilterCallbackBase* QueryCallback, const FQueryDebugParams& DebugParams)
+{
+	LowLevelRaycastImp(Scene, Start, Dir, DeltaMag, HitBuffer, OutputFlags, QueryFlags, Filter, QueryFilterData, QueryCallback, DebugParams);
+}
+
 void LowLevelSweep(FPhysScene& Scene, const FPhysicsGeometry& QueryGeom, const FTransform& StartTM, const FVector& Dir, float DeltaMag, FPhysicsHitCallback<FHitSweep>& HitBuffer, EHitFlags OutputFlags, FQueryFlags QueryFlags, const FCollisionFilterData& Filter, const FQueryFilterData& QueryFilterData, ICollisionQueryFilterCallbackBase* QueryCallback, const FQueryDebugParams& DebugParams)
 {
 #if !defined(PHYSICS_INTERFACE_PHYSX) || !PHYSICS_INTERFACE_PHYSX
diff --git a/Engine/Source/Runtime/Engine/Private/PhysicsEngine/CollisionQueryFilterCallback.cpp b/Engine/Source/Runtime/Engine/Private/PhysicsEngine/CollisionQueryFilterCallback.cpp
--- a/Engine/Source/Runtime/Engine/Private/PhysicsEngine/CollisionQueryFilterCallback.cpp
+++ b/Engine/Source/Runtime/Engine/Private/PhysicsEngine/CollisionQueryFilterCallback.cpp
@@ -109,7 +109,8 @@ ECollisionQueryHitType FCollisionQueryFilterCallback::PreFilterImp(const FCollis
 }
 #endif
 
-ECollisionQueryHitType FCollisionQueryFilterCallback::PreFilterImp(const FCollisionFilterData& FilterData, const Chaos::FPerShapeData& Shape, const Chaos::FGeometryParticle& Actor)
+template <typename TParticle>
+ECollisionQueryHitType FCollisionQueryFilterCallback::PreFilterBaseImp(const FCollisionFilterData& FilterData, const Chaos::FPerShapeData& Shape, const TParticle& Actor)
 {
 	//SCOPE_CYCLE_COUNTER(STAT_Collision_PreFilter);
 
@@ -128,13 +129,27 @@ ECollisionQueryHitType FCollisionQueryFilterCallback::PreFilterImp(const FCollis
 	}
 
 	FBodyInstance* BodyInstance = nullptr;
+
+	if constexpr (std::is_same<TParticle, Chaos::FGeometryParticle>::value)
+	{
 #if ENABLE_PREFILTER_LOGGING || DETECT_SQ_HITCHES
-	BodyInstance = ChaosInterface::GetUserData(Actor);
+		BodyInstance = ChaosInterface::GetUserData(Actor);
 #endif // ENABLE_PREFILTER_LOGGING || DETECT_SQ_HITCHES
+	}
 
 	return PreFilterImp(FilterData, ShapeFilter, ComponentID, BodyInstance);
 }
 
+ECollisionQueryHitType FCollisionQueryFilterCallback::PreFilterImp(const FCollisionFilterData& FilterData, const Chaos::FPerShapeData& Shape, const Chaos::FGeometryParticle& Actor)
+{
+	return PreFilterBaseImp(FilterData, Shape, Actor);
+}
+
+ECollisionQueryHitType FCollisionQueryFilterCallback::PreFilterImp(const FCollisionFilterData& FilterData, const Chaos::FPerShapeData& Shape, const Chaos::FGeometryParticleHandle& Actor)
+{
+	return PreFilterBaseImp(FilterData, Shape, Actor);
+}
+
 ECollisionQueryHitType FCollisionQueryFilterCallback::PreFilterImp(const FCollisionFilterData& FilterData, const FCollisionFilterData& ShapeFilter, uint32 ComponentID, const FBodyInstance* BodyInstance)
 {
 #if DETECT_SQ_HITCHES
@@ -282,12 +297,27 @@ ECollisionQueryHitType FCollisionQueryFilterCallback::PostFilterImp(const FColli
 		return ECollisionQueryHitType::None;
 	}
 
-	const ChaosInterface::FLocationHit& SweepHit = (const ChaosInterface::FLocationHit&)Hit;
+	const auto& SweepHit = static_cast<const ChaosInterface::FLocationHit&>(Hit);
+	const bool bIsOverlap = ChaosInterface::HadInitialOverlap(SweepHit);
+
+	return PostFilterImp(FilterData, bIsOverlap);
+}
+
+ECollisionQueryHitType FCollisionQueryFilterCallback::PostFilterImp(const FCollisionFilterData& FilterData, const ChaosInterface::FPTQueryHit& Hit)
+{
+	// Unused in non-sweeps
+	if (!bIsSweep)
+	{
+		return ECollisionQueryHitType::None;
+	}
+
+	const auto& SweepHit = static_cast<const ChaosInterface::FPTLocationHit&>(Hit);
 	const bool bIsOverlap = ChaosInterface::HadInitialOverlap(SweepHit);
 
 	return PostFilterImp(FilterData, bIsOverlap);
 }
 
+
 #if PHYSICS_INTERFACE_PHYSX
 
 PxQueryHitType::Enum FCollisionQueryFilterCallback::preFilter(const PxFilterData& filterData, const PxShape* shape, const PxRigidActor* actor, PxHitFlags& queryFlags)
diff --git a/Engine/Source/Runtime/Engine/Private/PhysicsEngine/CollisionQueryFilterCallback.h b/Engine/Source/Runtime/Engine/Private/PhysicsEngine/CollisionQueryFilterCallback.h
--- a/Engine/Source/Runtime/Engine/Private/PhysicsEngine/CollisionQueryFilterCallback.h
+++ b/Engine/Source/Runtime/Engine/Private/PhysicsEngine/CollisionQueryFilterCallback.h
@@ -70,17 +70,31 @@ public:
 #endif
 
 	ECollisionQueryHitType PreFilterImp(const FCollisionFilterData& FilterData, const Chaos::FPerShapeData& Shape, const Chaos::FGeometryParticle& Actor);
+	ECollisionQueryHitType PreFilterImp(const FCollisionFilterData& FilterData, const Chaos::FPerShapeData& Shape, const Chaos::FGeometryParticleHandle& Actor);
+
 	ECollisionQueryHitType PostFilterImp(const FCollisionFilterData& FilterData, const ChaosInterface::FQueryHit& Hit);
+	ECollisionQueryHitType PostFilterImp(const FCollisionFilterData& FilterData, const ChaosInterface::FPTQueryHit& Hit);
 
 	virtual ECollisionQueryHitType PostFilter(const FCollisionFilterData& FilterData, const ChaosInterface::FQueryHit& Hit) override
 	{
 		return PostFilterImp(FilterData, Hit);
 	}
+
+	virtual ECollisionQueryHitType PostFilter(const FCollisionFilterData& FilterData, const ChaosInterface::FPTQueryHit& Hit) override
+	{
+		return PostFilterImp(FilterData, Hit);
+	}
+
 	virtual ECollisionQueryHitType PreFilter(const FCollisionFilterData& FilterData, const Chaos::FPerShapeData& Shape, const Chaos::FGeometryParticle& Actor) override
 	{
 		return PreFilterImp(FilterData, Shape, Actor);
 	}
 
+	virtual ECollisionQueryHitType PreFilter(const FCollisionFilterData& FilterData, const Chaos::FPerShapeData& Shape, const Chaos::FGeometryParticleHandle& Actor) override
+	{
+		return PreFilterImp(FilterData, Shape, Actor);
+	}
+
 #if PHYSICS_INTERFACE_PHYSX
 	virtual ECollisionQueryHitType PostFilter(const FCollisionFilterData& FilterData, const physx::PxQueryHit& Hit) override { return PostFilterImp(FilterData, Hit); }
 	virtual ECollisionQueryHitType PreFilter(const FCollisionFilterData& FilterData, const physx::PxShape& Shape, physx::PxRigidActor& Actor) override { return PreFilterImp(FilterData, Shape, Actor); }
@@ -101,5 +115,9 @@ public:
 #endif
 	bool bDiscardInitialOverlaps;
 	bool bIsSweep;
+
+private:
+	template <typename TParticle>
+	ECollisionQueryHitType PreFilterBaseImp(const FCollisionFilterData& FilterData, const Chaos::FPerShapeData& Shape, const TParticle& Actor);
 };
 
diff --git a/Engine/Source/Runtime/Engine/Public/Physics/Experimental/ChaosInterfaceWrapper.h b/Engine/Source/Runtime/Engine/Public/Physics/Experimental/ChaosInterfaceWrapper.h
--- a/Engine/Source/Runtime/Engine/Public/Physics/Experimental/ChaosInterfaceWrapper.h
+++ b/Engine/Source/Runtime/Engine/Public/Physics/Experimental/ChaosInterfaceWrapper.h
@@ -51,5 +51,7 @@ UPhysicalMaterial* GetUserData(const Chaos::FChaosPhysicsMaterial& Material);
 }
 
 void LowLevelRaycast(FPhysScene& Scene, const FVector& Start, const FVector& Dir, float DeltaMag, FPhysicsHitCallback<FHitRaycast>& HitBuffer, EHitFlags OutputFlags, FQueryFlags QueryFlags, const FCollisionFilterData& Filter, const FQueryFilterData& QueryFilterData, ICollisionQueryFilterCallbackBase* QueryCallback, const FQueryDebugParams& DebugParams = FQueryDebugParams());
+void LowLevelRaycast(FPhysScene& Scene, const FVector& Start, const FVector& Dir, float DeltaMag, FPhysicsHitCallback<FPTRaycastHit>& HitBuffer, EHitFlags OutputFlags, FQueryFlags QueryFlags, const FCollisionFilterData& Filter, const FQueryFilterData& QueryFilterData, ICollisionQueryFilterCallbackBase* QueryCallback, const FQueryDebugParams& DebugParams = FQueryDebugParams());
+
 void LowLevelSweep(FPhysScene& Scene, const FPhysicsGeometry& Geom, const FTransform& StartTM, const FVector& Dir, float DeltaMag, FPhysicsHitCallback<FHitSweep>& HitBuffer, EHitFlags OutputFlags, FQueryFlags QueryFlags, const FCollisionFilterData& Filter, const FQueryFilterData& QueryFilterData, ICollisionQueryFilterCallbackBase* QueryCallback, const FQueryDebugParams& DebugParams = FQueryDebugParams());
 void LowLevelOverlap(FPhysScene& Scene, const FPhysicsGeometry& Geom, const FTransform& GeomPose, FPhysicsHitCallback<FHitOverlap>& HitBuffer, FQueryFlags QueryFlags, const FCollisionFilterData& Filter, const FQueryFilterData& QueryFilterData, ICollisionQueryFilterCallbackBase* QueryCallback, const FQueryDebugParams& DebugParams = FQueryDebugParams());
diff --git a/Engine/Source/Runtime/PhysicsCore/Private/PhysicsInterfaceUtilsCore.cpp b/Engine/Source/Runtime/PhysicsCore/Private/PhysicsInterfaceUtilsCore.cpp
--- a/Engine/Source/Runtime/PhysicsCore/Private/PhysicsInterfaceUtilsCore.cpp
+++ b/Engine/Source/Runtime/PhysicsCore/Private/PhysicsInterfaceUtilsCore.cpp
@@ -188,7 +188,8 @@ DEFINE_STAT(STAT_PhysSceneWriteLock);
 
 #endif
 
-uint32 FindFaceIndex(const FHitLocation& PHit, const FVector& UnitDir)
+template <typename THitLocation>
+uint32 FindFaceIndexImp(const THitLocation& PHit, const FVector& UnitDir)
 {
 #if PHYSICS_INTERFACE_PHYSX
 	PxConvexMeshGeometry convexGeom;
@@ -303,3 +304,13 @@ uint32 FindFaceIndex(const FHitLocation& PHit, const FVector& UnitDir)
 	return PHit.Shape->GetGeometry()->FindMostOpposingFace(LocalPosition, LocalNormal, PHit.FaceIndex, 1);	//todo:this number matches the one above, but is it right?
 #endif
 }
+
+uint32 FindFaceIndex(const FHitLocation& Hit, const FVector& UnitDir)
+{
+	return FindFaceIndexImp(Hit, UnitDir);
+}
+
+uint32 FindFaceIndex(const ChaosInterface::FPTLocationHit& Hit, const FVector& UnitDir)
+{
+	return FindFaceIndexImp(Hit, UnitDir);
+}
diff --git a/Engine/Source/Runtime/PhysicsCore/Private/SQAccelerator.cpp b/Engine/Source/Runtime/PhysicsCore/Private/SQAccelerator.cpp
--- a/Engine/Source/Runtime/PhysicsCore/Private/SQAccelerator.cpp
+++ b/Engine/Source/Runtime/PhysicsCore/Private/SQAccelerator.cpp
@@ -69,7 +69,8 @@ struct FPreFilterInfo
 	int32 ActorIdx;
 };
 
-void FillHitHelper(ChaosInterface::FLocationHit& Hit, const Chaos::FReal Distance, const FVector& WorldPosition, const FVector& WorldNormal, int32 FaceIdx, bool bComputeMTD)
+template <typename TLocationHit>
+void FillHitHelperImp(TLocationHit& Hit, const Chaos::FReal Distance, const FVector& WorldPosition, const FVector& WorldNormal, int32 FaceIdx, bool bComputeMTD)
 {
 	Hit.Distance = Distance;
 	Hit.WorldPosition = WorldPosition;
@@ -78,14 +79,25 @@ void FillHitHelper(ChaosInterface::FLocationHit& Hit, const Chaos::FReal Distanc
 	Hit.FaceIndex = FaceIdx;
 }
 
+void FillHitHelper(ChaosInterface::FLocationHit& Hit, const Chaos::FReal Distance, const FVector& WorldPosition, const FVector& WorldNormal, int32 FaceIdx, bool bComputeMTD)
+{
+	FillHitHelperImp(Hit, Distance, WorldPosition, WorldNormal, FaceIdx, bComputeMTD);
+}
+
+void FillHitHelper(ChaosInterface::FPTLocationHit& Hit, const Chaos::FReal Distance, const FVector& WorldPosition, const FVector& WorldNormal, int32 FaceIdx, bool bComputeMTD)
+{
+	FillHitHelperImp(Hit, Distance, WorldPosition, WorldNormal, FaceIdx, bComputeMTD);
+}
+
 void FillHitHelper(ChaosInterface::FOverlapHit& Hit, const Chaos::FReal Distance, const FVector& WorldPosition, const FVector& WorldNormal, int32 FaceIdx, bool bComputeMTD)
 {
 }
 
-template <typename QueryGeometryType, typename TPayload, typename THitType>
+template <typename QueryGeometryType, typename TPayload, typename THitType, bool bGTData = true>
 struct TSQVisitor : public Chaos::ISpatialVisitor<TPayload, Chaos::FReal>
 {
-	TSQVisitor(const FVector& InStartPoint, const FVector& InDir, ChaosInterface::FSQHitBuffer<ChaosInterface::FRaycastHit>& InHitBuffer, EHitFlags InOutputFlags,
+	using TGeometryType = typename TChooseClass<bGTData, Chaos::FGeometryParticle, Chaos::FGeometryParticleHandle>::Result;
+	TSQVisitor(const FVector& InStartPoint, const FVector& InDir, ChaosInterface::FSQHitBuffer<THitType>& InHitBuffer, EHitFlags InOutputFlags,
 		const FQueryFilterData& InQueryFilterData, ICollisionQueryFilterCallbackBase& InQueryCallback, const FQueryDebugParams& InDebugParams)
 		: StartPoint(InStartPoint)
 		, Dir(InDir)
@@ -201,6 +213,18 @@ private:
 		Overlap
 	};
 
+	auto GetPayloadForThread(const TPayload& Payload)
+	{
+		if constexpr(bGTData)
+		{
+			return Payload.GetExternalGeometryParticle_ExternalThread();
+		}
+		else
+		{
+			return Payload.GetGeometryParticleHandle_PhysicsThread();
+		}
+	}
+
 	template <ESQType SQ>
 	bool Visit(const Chaos::TSpatialVisitorData<TPayload>& Instance, Chaos::FQueryFastData* CurData)
 	{
@@ -209,12 +233,14 @@ private:
 
 		//todo: add a check to ensure hitbuffer matches SQ type
 		using namespace Chaos;
-		FGeometryParticle* GeometryParticle = Payload.GetExternalGeometryParticle_ExternalThread();
+		TGeometryType* GeometryParticle = GetPayloadForThread(Payload);
 
 		if(!GeometryParticle)
 		{
 			// This case handles particles created by the physics simulation without the main thread
 			// being made aware of their creation. We have a PT particle but no external particle
+			ensure(bGTData);
+
 			return true;
 		}
 
@@ -537,17 +563,28 @@ private:
 	const FTransform StartTM;
 };
 
-void FChaosSQAccelerator::Raycast(const FVector& Start, const FVector& Dir, const float DeltaMagnitude, ChaosInterface::FSQHitBuffer<ChaosInterface::FRaycastHit>& HitBuffer, EHitFlags OutputFlags, const FQueryFilterData& QueryFilterData, ICollisionQueryFilterCallbackBase& QueryCallback, const FQueryDebugParams& DebugParams) const
+template <typename TRaycastHit>
+void FChaosSQAccelerator::RaycastImp(const FVector& Start, const FVector& Dir, const float DeltaMagnitude, ChaosInterface::FSQHitBuffer<TRaycastHit>& HitBuffer, EHitFlags OutputFlags, const FQueryFilterData& QueryFilterData, ICollisionQueryFilterCallbackBase& QueryCallback, const FQueryDebugParams& DebugParams) const
 {
 	using namespace Chaos;
 	using namespace ChaosInterface;
 
-	TSQVisitor<TSphere<FReal,3>, FAccelerationStructureHandle, FRaycastHit> RaycastVisitor(Start, Dir, HitBuffer, OutputFlags, QueryFilterData, QueryCallback, DebugParams);
+	TSQVisitor<TSphere<FReal, 3>, FAccelerationStructureHandle, TRaycastHit, std::is_same<TRaycastHit,FRaycastHit>::value> RaycastVisitor(Start, Dir, HitBuffer, OutputFlags, QueryFilterData, QueryCallback, DebugParams);
 	HitBuffer.IncFlushCount();
 	SpatialAcceleration.Raycast(Start, Dir, DeltaMagnitude, RaycastVisitor);
 	HitBuffer.DecFlushCount();
 }
 
+void FChaosSQAccelerator::Raycast(const FVector& Start, const FVector& Dir, const float DeltaMagnitude, ChaosInterface::FSQHitBuffer<ChaosInterface::FRaycastHit>& HitBuffer, EHitFlags OutputFlags, const FQueryFilterData& QueryFilterData, ICollisionQueryFilterCallbackBase& QueryCallback, const FQueryDebugParams& DebugParams) const
+{
+	RaycastImp(Start, Dir, DeltaMagnitude, HitBuffer, OutputFlags, QueryFilterData, QueryCallback, DebugParams);
+}
+
+void FChaosSQAccelerator::Raycast(const FVector& Start, const FVector& Dir, const float DeltaMagnitude, ChaosInterface::FSQHitBuffer<ChaosInterface::FPTRaycastHit>& HitBuffer, EHitFlags OutputFlags, const FQueryFilterData& QueryFilterData, ICollisionQueryFilterCallbackBase& QueryCallback, const FQueryDebugParams& DebugParams) const
+{
+	RaycastImp(Start, Dir, DeltaMagnitude, HitBuffer, OutputFlags, QueryFilterData, QueryCallback, DebugParams);
+}
+
 template <typename QueryGeomType>
 void SweepHelper(const QueryGeomType& QueryGeom,const Chaos::ISpatialAcceleration<Chaos::FAccelerationStructureHandle, Chaos::FReal,3>& SpatialAcceleration,const FTransform& StartTM,const FVector& Dir,const float DeltaMagnitude,ChaosInterface::FSQHitBuffer<ChaosInterface::FSweepHit>& HitBuffer,EHitFlags OutputFlags,const FQueryFilterData& QueryFilterData,ICollisionQueryFilterCallbackBase& QueryCallback,const FQueryDebugParams& DebugParams)
 {
diff --git a/Engine/Source/Runtime/PhysicsCore/Private/SQCapture.cpp b/Engine/Source/Runtime/PhysicsCore/Private/SQCapture.cpp
--- a/Engine/Source/Runtime/PhysicsCore/Private/SQCapture.cpp
+++ b/Engine/Source/Runtime/PhysicsCore/Private/SQCapture.cpp
@@ -22,6 +22,9 @@ public:
 	virtual ECollisionQueryHitType PostFilter(const FCollisionFilterData& FilterData, const ChaosInterface::FQueryHit& Hit) override { /*check(false);*/  return ECollisionQueryHitType::Touch; }
 	virtual ECollisionQueryHitType PreFilter(const FCollisionFilterData& FilterData, const Chaos::FPerShapeData& Shape, const Chaos::FGeometryParticle& Actor) override { return Capture.GetFilterResult(&Shape, &Actor); }
 
+	virtual ECollisionQueryHitType PostFilter(const FCollisionFilterData& FilterData, const ChaosInterface::FPTQueryHit& Hit) override { ensure(false);  return ECollisionQueryHitType::Touch; }
+	virtual ECollisionQueryHitType PreFilter(const FCollisionFilterData& FilterData, const Chaos::FPerShapeData& Shape, const Chaos::FGeometryParticleHandle& Actor) override { ensure(false);  return ECollisionQueryHitType::Touch; }
+
 #if PHYSICS_INTERFACE_PHYSX
 	virtual ECollisionQueryHitType PostFilter(const FCollisionFilterData& FilterData, const physx::PxQueryHit& Hit) override { return ECollisionQueryHitType::Touch; }
 	virtual ECollisionQueryHitType PreFilter(const FCollisionFilterData& FilterData, const physx::PxShape& Shape, physx::PxRigidActor& Actor) override { return Capture.GetFilterResult(&Shape, &Actor); }
diff --git a/Engine/Source/Runtime/PhysicsCore/Public/CollisionQueryFilterCallbackCore.h b/Engine/Source/Runtime/PhysicsCore/Public/CollisionQueryFilterCallbackCore.h
--- a/Engine/Source/Runtime/PhysicsCore/Public/CollisionQueryFilterCallbackCore.h
+++ b/Engine/Source/Runtime/PhysicsCore/Public/CollisionQueryFilterCallbackCore.h
@@ -41,6 +41,9 @@ public:
 	virtual ECollisionQueryHitType PostFilter(const FCollisionFilterData& FilterData, const ChaosInterface::FQueryHit& Hit) = 0;
 	virtual ECollisionQueryHitType PreFilter(const FCollisionFilterData& FilterData, const Chaos::FPerShapeData& Shape, const Chaos::FGeometryParticle& Actor) = 0;
 
+	virtual ECollisionQueryHitType PostFilter(const FCollisionFilterData& FilterData, const ChaosInterface::FPTQueryHit& Hit) = 0;
+	virtual ECollisionQueryHitType PreFilter(const FCollisionFilterData& FilterData, const Chaos::FPerShapeData& Shape, const Chaos::FGeometryParticleHandle& Actor) = 0;
+
 #if  PHYSICS_INTERFACE_PHYSX
 	virtual ECollisionQueryHitType PostFilter(const FCollisionFilterData& FilterData, const physx::PxQueryHit& Hit) = 0;
 	virtual ECollisionQueryHitType PreFilter(const FCollisionFilterData& FilterData, const physx::PxShape& Shape, physx::PxRigidActor& Actor) = 0;
@@ -54,6 +57,9 @@ public:
 	virtual ECollisionQueryHitType PostFilter(const FCollisionFilterData& FilterData, const ChaosInterface::FQueryHit& Hit) override { return ECollisionQueryHitType::Block; }
 	virtual ECollisionQueryHitType PreFilter(const FCollisionFilterData& FilterData, const Chaos::FPerShapeData& Shape, const Chaos::FGeometryParticle& Actor) override { return ECollisionQueryHitType::Block; }
 
+	virtual ECollisionQueryHitType PostFilter(const FCollisionFilterData& FilterData, const ChaosInterface::FPTQueryHit& Hit) override { return ECollisionQueryHitType::Block; }
+	virtual ECollisionQueryHitType PreFilter(const FCollisionFilterData& FilterData, const Chaos::FPerShapeData& Shape, const Chaos::FGeometryParticleHandle& Actor) override { return ECollisionQueryHitType::Block; }
+
 #if PHYSICS_INTERFACE_PHYSX
 	virtual ECollisionQueryHitType PostFilter(const FCollisionFilterData& FilterData, const physx::PxQueryHit& Hit) override { return ECollisionQueryHitType::Block; }
 	virtual ECollisionQueryHitType PreFilter(const FCollisionFilterData& FilterData, const physx::PxShape& Shape, physx::PxRigidActor& Actor) override { return ECollisionQueryHitType::Block; }
@@ -69,6 +75,9 @@ public:
 	virtual ECollisionQueryHitType PostFilter(const FCollisionFilterData& FilterData, const ChaosInterface::FQueryHit& Hit) override { return ECollisionQueryHitType::Touch; }
 	virtual ECollisionQueryHitType PreFilter(const FCollisionFilterData& FilterData, const Chaos::FPerShapeData& Shape, const Chaos::FGeometryParticle& Actor) override { return ECollisionQueryHitType::Touch; }
 
+	virtual ECollisionQueryHitType PostFilter(const FCollisionFilterData& FilterData, const ChaosInterface::FPTQueryHit& Hit) override { return ECollisionQueryHitType::Touch; }
+	virtual ECollisionQueryHitType PreFilter(const FCollisionFilterData& FilterData, const Chaos::FPerShapeData& Shape, const Chaos::FGeometryParticleHandle& Actor) override { return ECollisionQueryHitType::Touch; }
+
 #if PHYSICS_INTERFACE_PHYSX
 	virtual ECollisionQueryHitType PostFilter(const FCollisionFilterData& FilterData, const physx::PxQueryHit& Hit) override { return ECollisionQueryHitType::Touch; }
 	virtual ECollisionQueryHitType PreFilter(const FCollisionFilterData& FilterData, const physx::PxShape& Shape, physx::PxRigidActor& Actor) override { return ECollisionQueryHitType::Touch; }
diff --git a/Engine/Source/Runtime/PhysicsCore/Public/PhysicsInterfaceUtilsCore.h b/Engine/Source/Runtime/PhysicsCore/Public/PhysicsInterfaceUtilsCore.h
--- a/Engine/Source/Runtime/PhysicsCore/Public/PhysicsInterfaceUtilsCore.h
+++ b/Engine/Source/Runtime/PhysicsCore/Public/PhysicsInterfaceUtilsCore.h
@@ -176,4 +176,6 @@ private:
 #endif // WITH_PHYX
 
 //Find the face index for a given hit. This gives us a chance to modify face index based on things like most opposing normal
+
 PHYSICSCORE_API uint32 FindFaceIndex(const FHitLocation& PHit, const FVector& UnitDirection);
+PHYSICSCORE_API uint32 FindFaceIndex(const ChaosInterface::FPTLocationHit& PHit, const FVector& UnitDirection);
diff --git a/Engine/Source/Runtime/PhysicsCore/Public/SQAccelerator.h b/Engine/Source/Runtime/PhysicsCore/Public/SQAccelerator.h
--- a/Engine/Source/Runtime/PhysicsCore/Public/SQAccelerator.h
+++ b/Engine/Source/Runtime/PhysicsCore/Public/SQAccelerator.h
@@ -43,12 +43,17 @@ public:
 	virtual ~FChaosSQAccelerator() {};
 
 	void Raycast(const FVector& Start, const FVector& Dir, const float DeltaMagnitude, ChaosInterface::FSQHitBuffer<ChaosInterface::FRaycastHit>& HitBuffer, EHitFlags OutputFlags, const FQueryFilterData& QueryFilterData, ICollisionQueryFilterCallbackBase& QueryCallback, const FQueryDebugParams& DebugParams = FQueryDebugParams()) const;
+	void Raycast(const FVector& Start, const FVector& Dir, const float DeltaMagnitude, ChaosInterface::FSQHitBuffer<ChaosInterface::FPTRaycastHit>& HitBuffer, EHitFlags OutputFlags, const FQueryFilterData& QueryFilterData, ICollisionQueryFilterCallbackBase& QueryCallback, const FQueryDebugParams& DebugParams = FQueryDebugParams()) const;
+
 	void Sweep(const Chaos::FImplicitObject& QueryGeom, const FTransform& StartTM, const FVector& Dir, const float DeltaMagnitude, ChaosInterface::FSQHitBuffer<ChaosInterface::FSweepHit>& HitBuffer, EHitFlags OutputFlags, const FQueryFilterData& QueryFilterData, ICollisionQueryFilterCallbackBase& QueryCallback, const FQueryDebugParams& DebugParams = FQueryDebugParams()) const;
 	void Overlap(const Chaos::FImplicitObject& QueryGeom, const FTransform& GeomPose, ChaosInterface::FSQHitBuffer<ChaosInterface::FOverlapHit>& HitBuffer, const FQueryFilterData& QueryFilterData, ICollisionQueryFilterCallbackBase& QueryCallback, const FQueryDebugParams& DebugParams = FQueryDebugParams()) const;
 
 private:
 	const Chaos::ISpatialAcceleration<Chaos::FAccelerationStructureHandle, Chaos::FReal, 3>& SpatialAcceleration;
 
+	template <typename TRaycastHit>
+	void RaycastImp(const FVector& Start, const FVector& Dir, const float DeltaMagnitude, ChaosInterface::FSQHitBuffer<TRaycastHit>& HitBuffer, EHitFlags OutputFlags, const FQueryFilterData& QueryFilterData, ICollisionQueryFilterCallbackBase& QueryCallback, const FQueryDebugParams& DebugParams) const;
+
 };
 
 // An interface to the scene query accelerator that allows us to run queries against either PhysX or Chaos
