Added templated boilerplate for sweep and overlap on PT

diff --git a/Engine/Source/Runtime/Engine/Private/Collision/CollisionConversions.cpp b/Engine/Source/Runtime/Engine/Private/Collision/CollisionConversions.cpp
--- a/Engine/Source/Runtime/Engine/Private/Collision/CollisionConversions.cpp
+++ b/Engine/Source/Runtime/Engine/Private/Collision/CollisionConversions.cpp
@@ -506,6 +506,8 @@ EConvertQueryResult ConvertTraceResults(bool& OutHasValidBlockingHit, const UWor
 
 template EConvertQueryResult ConvertTraceResults<FHitSweep>(bool& OutHasValidBlockingHit, const UWorld* World, int32 NumHits, FHitSweep* Hits, float CheckLength, const FCollisionFilterData& QueryFilter, TArray<FHitResult>& OutHits, const FVector& StartLoc, const FVector& EndLoc, const FPhysicsGeometry& Geom, const FTransform& QueryTM, float MaxDistance, bool bReturnFaceIndex, bool bReturnPhysMat);
 template EConvertQueryResult ConvertTraceResults<FHitSweep>(bool& OutHasValidBlockingHit, const UWorld* World, int32 NumHits, FHitSweep* Hits, float CheckLength, const FCollisionFilterData& QueryFilter, FHitResult& OutHit, const FVector& StartLoc, const FVector& EndLoc, const FPhysicsGeometry& Geom, const FTransform& QueryTM, float MaxDistance, bool bReturnFaceIndex, bool bReturnPhysMat);
+template EConvertQueryResult ConvertTraceResults<FPTSweepHit>(bool& OutHasValidBlockingHit, const UWorld* World, int32 NumHits, FPTSweepHit* Hits, float CheckLength, const FCollisionFilterData& QueryFilter, TArray<FHitResult>& OutHits, const FVector& StartLoc, const FVector& EndLoc, const FPhysicsGeometry& Geom, const FTransform& QueryTM, float MaxDistance, bool bReturnFaceIndex, bool bReturnPhysMat);
+template EConvertQueryResult ConvertTraceResults<FPTSweepHit>(bool& OutHasValidBlockingHit, const UWorld* World, int32 NumHits, FPTSweepHit* Hits, float CheckLength, const FCollisionFilterData& QueryFilter, FHitResult& OutHit, const FVector& StartLoc, const FVector& EndLoc, const FPhysicsGeometry& Geom, const FTransform& QueryTM, float MaxDistance, bool bReturnFaceIndex, bool bReturnPhysMat);
 template EConvertQueryResult ConvertTraceResults<FHitRaycast>(bool& OutHasValidBlockingHit, const UWorld* World, int32 NumHits, FHitRaycast* Hits, float CheckLength, const FCollisionFilterData& QueryFilter, TArray<FHitResult>& OutHits, const FVector& StartLoc, const FVector& EndLoc, const FPhysicsGeometry& Geom, const FTransform& QueryTM, float MaxDistance, bool bReturnFaceIndex, bool bReturnPhysMat);
 template EConvertQueryResult ConvertTraceResults<FHitRaycast>(bool& OutHasValidBlockingHit, const UWorld* World, int32 NumHits, FHitRaycast* Hits, float CheckLength, const FCollisionFilterData& QueryFilter, FHitResult& OutHit, const FVector& StartLoc, const FVector& EndLoc, const FPhysicsGeometry& Geom, const FTransform& QueryTM, float MaxDistance, bool bReturnFaceIndex, bool bReturnPhysMat);
 template EConvertQueryResult ConvertTraceResults<FPTRaycastHit>(bool& OutHasValidBlockingHit, const UWorld* World, int32 NumHits, FPTRaycastHit* Hits, float CheckLength, const FCollisionFilterData& QueryFilter, TArray<FHitResult>& OutHits, const FVector& StartLoc, const FVector& EndLoc, const FPhysicsGeometry& Geom, const FTransform& QueryTM, float MaxDistance, bool bReturnFaceIndex, bool bReturnPhysMat);
@@ -629,8 +631,6 @@ static bool ConvertOverlappedShapeToImpactHit(const UWorld* World, const THitLoc
 
 void ConvertQueryOverlap(const FPhysicsShape& Shape, const FPhysicsActor& Actor, FOverlapResult& OutOverlap, const FCollisionFilterData& QueryFilter)
 {
-	const bool bBlock = IsBlocking(Shape, QueryFilter);
-
 	// Grab actor/component
 	
 	// Try body instance
@@ -676,9 +676,6 @@ void ConvertQueryOverlap(const FPhysicsShape& Shape, const FPhysicsActor& Actor,
 			ensureMsgf(false, TEXT("ConvertQueryOverlap called with bad payload type"));
 		}
 	}
-
-	// Other info
-	OutOverlap.bBlockingHit = bBlock;
 }
 
 /** Util to add NewOverlap to OutOverlaps if it is not already there */
@@ -726,13 +723,15 @@ static FAutoConsoleVariableRef GTestOverlapSpeed(
 	TEXT("Min number of overlaps required before using a TMap for deduplication")
 	);
 
-bool ConvertOverlapResults(int32 NumOverlaps, FHitOverlap* OverlapResults, const FCollisionFilterData& QueryFilter, TArray<FOverlapResult>& OutOverlaps)
+template <typename THitOverlap>
+bool ConvertOverlapResultsImp(int32 NumOverlaps, THitOverlap* OverlapResults, const FCollisionFilterData& QueryFilter, TArray<FOverlapResult>& OutOverlaps)
 {
 	SCOPE_CYCLE_COUNTER(STAT_CollisionConvertOverlap);
 
 	const int32 ExpectedSize = OutOverlaps.Num() + NumOverlaps;
 	OutOverlaps.Reserve(ExpectedSize);
 	bool bBlockingFound = false;
+	constexpr bool bIsGTQuery = std::is_same<THitOverlap, FHitOverlap>::value;
 
 	if (ExpectedSize >= GNumOverlapsRequiredForTMap)
 	{
@@ -750,14 +749,24 @@ bool ConvertOverlapResults(int32 NumOverlaps, FHitOverlap* OverlapResults, const
 		for (int32 PResultIndex = 0; PResultIndex < NumOverlaps; ++PResultIndex)
 		{
 			FOverlapResult NewOverlap;
-			ConvertQueryOverlap(*GetShape(OverlapResults[PResultIndex]), *GetActor(OverlapResults[PResultIndex]), NewOverlap, QueryFilter);
+			const auto Actor = GetActor(OverlapResults[PResultIndex]);
+			const FPhysicsShape* Shape = GetShape(OverlapResults[PResultIndex]);
 
-			if (NewOverlap.bBlockingHit)
+			const bool bBlock = IsBlocking(*Shape, QueryFilter);
+			NewOverlap.bBlockingHit = bBlock;
+			bBlockingFound |= bBlock;
+
+			const FPhysicsActor* GTActor = GetGTActor(Actor);
+			const FPhysicsShape* GTShape = GetGTShape<bIsGTQuery>(Shape, GTActor);
+			if(GTActor && GTShape)
 			{
-				bBlockingFound = true;
+				ConvertQueryOverlap(*GTShape, *GTActor, NewOverlap, QueryFilter);
 			}
 
 			// Look for it in the map, newly added elements will start with 0, so we know we need to add it to the results array then (the index is stored as +1)
+			//TODO: this doesn't de-duplicate if FOverlapResult has no component - like in PT for example. If we care about de-duplication we should use the PT handle if needed
+			//Question: why do we de-duplicate? This seems expensive and is not consistent with traces. The user could handle this at their level anyway (at the cost of a bit of transient memory here)
+			//TODO: make sure it doesn't de-duplicate null results
 			int32& DestinationIndex = OverlapMap.FindOrAdd(FOverlapKey(NewOverlap.Component.Get(), NewOverlap.ItemIndex));
 			if (DestinationIndex == 0)
 			{
@@ -781,11 +790,18 @@ bool ConvertOverlapResults(int32 NumOverlaps, FHitOverlap* OverlapResults, const
 		for (int32 i = 0; i < NumOverlaps; i++)
 		{
 			FOverlapResult NewOverlap;
-			ConvertQueryOverlap(*GetShape(OverlapResults[i]), *GetActor(OverlapResults[i]), NewOverlap, QueryFilter);
+			const auto Actor = GetActor(OverlapResults[i]);
+			const FPhysicsShape* Shape = GetShape(OverlapResults[i]);
+
+			const bool bBlock = IsBlocking(*Shape, QueryFilter);
+			NewOverlap.bBlockingHit = bBlock;
+			bBlockingFound |= bBlock;
 
-			if (NewOverlap.bBlockingHit)
+			const FPhysicsActor* GTActor = GetGTActor(Actor);
+			const FPhysicsShape* GTShape = GetGTShape<bIsGTQuery>(Shape, GTActor);
+			if (GTActor && GTShape)
 			{
-				bBlockingFound = true;
+				ConvertQueryOverlap(*GTShape, *GTActor, NewOverlap, QueryFilter);
 			}
 
 			AddUniqueOverlap(OutOverlaps, NewOverlap);
@@ -794,3 +810,13 @@ bool ConvertOverlapResults(int32 NumOverlaps, FHitOverlap* OverlapResults, const
 
 	return bBlockingFound;
 }
+
+bool ConvertOverlapResults(int32 NumOverlaps, FHitOverlap* OverlapResults, const FCollisionFilterData& QueryFilter, TArray<FOverlapResult>& OutOverlaps)
+{
+	return ConvertOverlapResultsImp(NumOverlaps, OverlapResults, QueryFilter, OutOverlaps);
+}
+
+bool ConvertOverlapResults(int32 NumOverlaps, FPTOverlapHit* OverlapResults, const FCollisionFilterData& QueryFilter, TArray<FOverlapResult>& OutOverlaps)
+{
+	return ConvertOverlapResultsImp(NumOverlaps, OverlapResults, QueryFilter, OutOverlaps);
+}
\ No newline at end of file
diff --git a/Engine/Source/Runtime/Engine/Private/Collision/SceneQueryLowLevel.cpp b/Engine/Source/Runtime/Engine/Private/Collision/SceneQueryLowLevel.cpp
--- a/Engine/Source/Runtime/Engine/Private/Collision/SceneQueryLowLevel.cpp
+++ b/Engine/Source/Runtime/Engine/Private/Collision/SceneQueryLowLevel.cpp
@@ -233,7 +233,8 @@ void LowLevelRaycast(FPhysScene& Scene, const FVector& Start, const FVector& Dir
 	LowLevelRaycastImp(Scene, Start, Dir, DeltaMag, HitBuffer, OutputFlags, QueryFlags, Filter, QueryFilterData, QueryCallback, DebugParams);
 }
 
-void LowLevelSweep(FPhysScene& Scene, const FPhysicsGeometry& QueryGeom, const FTransform& StartTM, const FVector& Dir, float DeltaMag, FPhysicsHitCallback<FHitSweep>& HitBuffer, EHitFlags OutputFlags, FQueryFlags QueryFlags, const FCollisionFilterData& Filter, const FQueryFilterData& QueryFilterData, ICollisionQueryFilterCallbackBase* QueryCallback, const FQueryDebugParams& DebugParams)
+template <typename THitSweep>
+void LowLevelSweepImp(FPhysScene& Scene, const FPhysicsGeometry& QueryGeom, const FTransform& StartTM, const FVector& Dir, float DeltaMag, FPhysicsHitCallback<THitSweep>& HitBuffer, EHitFlags OutputFlags, FQueryFlags QueryFlags, const FCollisionFilterData& Filter, const FQueryFilterData& QueryFilterData, ICollisionQueryFilterCallbackBase* QueryCallback, const FQueryDebugParams& DebugParams)
 {
 #if !defined(PHYSICS_INTERFACE_PHYSX) || !PHYSICS_INTERFACE_PHYSX
 	if (const auto& SolverAccelerationStructure = Scene.GetSpacialAcceleration())
@@ -247,9 +248,12 @@ void LowLevelSweep(FPhysScene& Scene, const FPhysicsGeometry& QueryGeom, const F
 				SQAccelerator.Sweep(QueryGeom, StartTM, Dir, DeltaMag, HitBuffer, OutputFlags, QueryFilterData, *QueryCallback, DebugParams);
 			}
 
-			if (!!SerializeSQs && !!EnableSweepSQCapture)
+			if constexpr (std::is_same<THitSweep, FHitSweep>::value)
 			{
-				SweepSQCaptureHelper(Time, SQAccelerator, Scene, QueryGeom, StartTM, Dir, DeltaMag, HitBuffer, OutputFlags, QueryFlags, Filter, QueryFilterData, QueryCallback, DebugParams);
+				if (!!SerializeSQs && !!EnableSweepSQCapture)
+				{
+					SweepSQCaptureHelper(Time, SQAccelerator, Scene, QueryGeom, StartTM, Dir, DeltaMag, HitBuffer, OutputFlags, QueryFlags, Filter, QueryFilterData, QueryCallback, DebugParams);
+				}
 			}
 		}
 	}
@@ -272,7 +276,18 @@ void LowLevelSweep(FPhysScene& Scene, const FPhysicsGeometry& QueryGeom, const F
 #endif
 }
 
-void LowLevelOverlap(FPhysScene& Scene, const FPhysicsGeometry& QueryGeom, const FTransform& GeomPose, FPhysicsHitCallback<FHitOverlap>& HitBuffer, FQueryFlags QueryFlags, const FCollisionFilterData& Filter, const FQueryFilterData& QueryFilterData, ICollisionQueryFilterCallbackBase* QueryCallback, const FQueryDebugParams& DebugParams)
+void LowLevelSweep(FPhysScene& Scene, const FPhysicsGeometry& QueryGeom, const FTransform& StartTM, const FVector& Dir, float DeltaMag, FPhysicsHitCallback<FHitSweep>& HitBuffer, EHitFlags OutputFlags, FQueryFlags QueryFlags, const FCollisionFilterData& Filter, const FQueryFilterData& QueryFilterData, ICollisionQueryFilterCallbackBase* QueryCallback, const FQueryDebugParams& DebugParams)
+{
+	return LowLevelSweepImp(Scene, QueryGeom, StartTM, Dir, DeltaMag, HitBuffer, OutputFlags, QueryFlags, Filter, QueryFilterData, QueryCallback, DebugParams);
+}
+
+void LowLevelSweep(FPhysScene& Scene, const FPhysicsGeometry& QueryGeom, const FTransform& StartTM, const FVector& Dir, float DeltaMag, FPhysicsHitCallback<FPTSweepHit>& HitBuffer, EHitFlags OutputFlags, FQueryFlags QueryFlags, const FCollisionFilterData& Filter, const FQueryFilterData& QueryFilterData, ICollisionQueryFilterCallbackBase* QueryCallback, const FQueryDebugParams& DebugParams)
+{
+	return LowLevelSweepImp(Scene, QueryGeom, StartTM, Dir, DeltaMag, HitBuffer, OutputFlags, QueryFlags, Filter, QueryFilterData, QueryCallback, DebugParams);
+}
+
+template <typename THitOverlap>
+void LowLevelOverlapImp(FPhysScene& Scene, const FPhysicsGeometry& QueryGeom, const FTransform& GeomPose, FPhysicsHitCallback<THitOverlap>& HitBuffer, FQueryFlags QueryFlags, const FCollisionFilterData& Filter, const FQueryFilterData& QueryFilterData, ICollisionQueryFilterCallbackBase* QueryCallback, const FQueryDebugParams& DebugParams)
 {
 #if !defined(PHYSICS_INTERFACE_PHYSX) || !PHYSICS_INTERFACE_PHYSX
 	if (const auto& SolverAccelerationStructure = Scene.GetSpacialAcceleration())
@@ -284,9 +299,12 @@ void LowLevelOverlap(FPhysScene& Scene, const FPhysicsGeometry& QueryGeom, const
 			SQAccelerator.Overlap(QueryGeom, GeomPose, HitBuffer, QueryFilterData, *QueryCallback, DebugParams);
 		}
 
-		if (!!SerializeSQs && !!EnableOverlapSQCapture)
+		if constexpr (std::is_same<THitOverlap, FHitOverlap>::value)
 		{
-			OverlapSQCaptureHelper(Time, SQAccelerator, Scene, QueryGeom, GeomPose, HitBuffer, QueryFlags, Filter, QueryFilterData, QueryCallback, DebugParams);
+			if (!!SerializeSQs && !!EnableOverlapSQCapture)
+			{
+				OverlapSQCaptureHelper(Time, SQAccelerator, Scene, QueryGeom, GeomPose, HitBuffer, QueryFlags, Filter, QueryFilterData, QueryCallback, DebugParams);
+			}
 		}
 	}
 #else
@@ -308,4 +326,14 @@ void LowLevelOverlap(FPhysScene& Scene, const FPhysicsGeometry& QueryGeom, const
 #endif
 }
 
+void LowLevelOverlap(FPhysScene& Scene, const FPhysicsGeometry& QueryGeom, const FTransform& GeomPose, FPhysicsHitCallback<FHitOverlap>& HitBuffer, FQueryFlags QueryFlags, const FCollisionFilterData& Filter, const FQueryFilterData& QueryFilterData, ICollisionQueryFilterCallbackBase* QueryCallback, const FQueryDebugParams& DebugParams)
+{
+	return LowLevelOverlapImp(Scene, QueryGeom, GeomPose, HitBuffer, QueryFlags, Filter, QueryFilterData, QueryCallback, DebugParams);
+}
+
+void LowLevelOverlap(FPhysScene& Scene, const FPhysicsGeometry& QueryGeom, const FTransform& GeomPose, FPhysicsHitCallback<FPTOverlapHit>& HitBuffer, FQueryFlags QueryFlags, const FCollisionFilterData& Filter, const FQueryFilterData& QueryFilterData, ICollisionQueryFilterCallbackBase* QueryCallback, const FQueryDebugParams& DebugParams)
+{
+	return LowLevelOverlapImp(Scene, QueryGeom, GeomPose, HitBuffer, QueryFlags, Filter, QueryFilterData, QueryCallback, DebugParams);
+}
+
 //#endif // WITH_PHYSX 
\ No newline at end of file
diff --git a/Engine/Source/Runtime/Engine/Public/Physics/Experimental/ChaosInterfaceWrapper.h b/Engine/Source/Runtime/Engine/Public/Physics/Experimental/ChaosInterfaceWrapper.h
--- a/Engine/Source/Runtime/Engine/Public/Physics/Experimental/ChaosInterfaceWrapper.h
+++ b/Engine/Source/Runtime/Engine/Public/Physics/Experimental/ChaosInterfaceWrapper.h
@@ -54,4 +54,7 @@ void LowLevelRaycast(FPhysScene& Scene, const FVector& Start, const FVector& Dir
 void LowLevelRaycast(FPhysScene& Scene, const FVector& Start, const FVector& Dir, float DeltaMag, FPhysicsHitCallback<FPTRaycastHit>& HitBuffer, EHitFlags OutputFlags, FQueryFlags QueryFlags, const FCollisionFilterData& Filter, const FQueryFilterData& QueryFilterData, ICollisionQueryFilterCallbackBase* QueryCallback, const FQueryDebugParams& DebugParams = FQueryDebugParams());
 
 void LowLevelSweep(FPhysScene& Scene, const FPhysicsGeometry& Geom, const FTransform& StartTM, const FVector& Dir, float DeltaMag, FPhysicsHitCallback<FHitSweep>& HitBuffer, EHitFlags OutputFlags, FQueryFlags QueryFlags, const FCollisionFilterData& Filter, const FQueryFilterData& QueryFilterData, ICollisionQueryFilterCallbackBase* QueryCallback, const FQueryDebugParams& DebugParams = FQueryDebugParams());
+void LowLevelSweep(FPhysScene& Scene, const FPhysicsGeometry& Geom, const FTransform& StartTM, const FVector& Dir, float DeltaMag, FPhysicsHitCallback<FPTSweepHit>& HitBuffer, EHitFlags OutputFlags, FQueryFlags QueryFlags, const FCollisionFilterData& Filter, const FQueryFilterData& QueryFilterData, ICollisionQueryFilterCallbackBase* QueryCallback, const FQueryDebugParams& DebugParams = FQueryDebugParams());
+
 void LowLevelOverlap(FPhysScene& Scene, const FPhysicsGeometry& Geom, const FTransform& GeomPose, FPhysicsHitCallback<FHitOverlap>& HitBuffer, FQueryFlags QueryFlags, const FCollisionFilterData& Filter, const FQueryFilterData& QueryFilterData, ICollisionQueryFilterCallbackBase* QueryCallback, const FQueryDebugParams& DebugParams = FQueryDebugParams());
+void LowLevelOverlap(FPhysScene& Scene, const FPhysicsGeometry& Geom, const FTransform& GeomPose, FPhysicsHitCallback<FPTOverlapHit>& HitBuffer, FQueryFlags QueryFlags, const FCollisionFilterData& Filter, const FQueryFilterData& QueryFilterData, ICollisionQueryFilterCallbackBase* QueryCallback, const FQueryDebugParams& DebugParams = FQueryDebugParams());
diff --git a/Engine/Source/Runtime/PhysicsCore/Private/SQAccelerator.cpp b/Engine/Source/Runtime/PhysicsCore/Private/SQAccelerator.cpp
--- a/Engine/Source/Runtime/PhysicsCore/Private/SQAccelerator.cpp
+++ b/Engine/Source/Runtime/PhysicsCore/Private/SQAccelerator.cpp
@@ -93,6 +93,10 @@ void FillHitHelper(ChaosInterface::FOverlapHit& Hit, const Chaos::FReal Distance
 {
 }
 
+void FillHitHelper(ChaosInterface::FPTOverlapHit& Hit, const Chaos::FReal Distance, const FVector& WorldPosition, const FVector& WorldNormal, int32 FaceIdx, bool bComputeMTD)
+{
+}
+
 template <typename QueryGeometryType, typename TPayload, typename THitType, bool bGTData = true>
 struct TSQVisitor : public Chaos::ISpatialVisitor<TPayload, Chaos::FReal>
 {
@@ -121,7 +125,7 @@ struct TSQVisitor : public Chaos::ISpatialVisitor<TPayload, Chaos::FReal>
 #endif
 	}
 
-	TSQVisitor(const FTransform& InStartTM, const FVector& InDir, ChaosInterface::FSQHitBuffer<ChaosInterface::FSweepHit>& InHitBuffer, EHitFlags InOutputFlags,
+	TSQVisitor(const FTransform& InStartTM, const FVector& InDir, ChaosInterface::FSQHitBuffer<THitType>& InHitBuffer, EHitFlags InOutputFlags,
 		const FQueryFilterData& InQueryFilterData, ICollisionQueryFilterCallbackBase& InQueryCallback, const QueryGeometryType& InQueryGeom, const FQueryDebugParams& InDebugParams)
 		: Dir(InDir)
 		, HalfExtents(InQueryGeom.BoundingBox().Extents() * 0.5)
@@ -148,7 +152,7 @@ struct TSQVisitor : public Chaos::ISpatialVisitor<TPayload, Chaos::FReal>
 		//todo: check THitType is sweep
 	}
 
-	TSQVisitor(const FTransform& InWorldTM, ChaosInterface::FSQHitBuffer<ChaosInterface::FOverlapHit>& InHitBuffer,
+	TSQVisitor(const FTransform& InWorldTM, ChaosInterface::FSQHitBuffer<THitType>& InHitBuffer,
 		const FQueryFilterData& InQueryFilterData, ICollisionQueryFilterCallbackBase& InQueryCallback, const QueryGeometryType& InQueryGeom, const FQueryDebugParams& InDebugParams)
 		: HalfExtents(InQueryGeom.BoundingBox().Extents() * 0.5)
 		, bAnyHit(false)
@@ -455,10 +459,12 @@ private:
 	const FTransform StartTM;
 };
 
-template <typename QueryGeometryType, typename TPayload, typename THitType>
+template <typename QueryGeometryType, typename TPayload, typename THitType, bool bGTData = true>
 struct TBPVisitor : public Chaos::ISpatialVisitor<TPayload, Chaos::FReal>
 {
-	TBPVisitor(const FTransform& InWorldTM, ChaosInterface::FSQHitBuffer<ChaosInterface::FOverlapHit>& InHitBuffer,
+	using TGeometryType = typename TChooseClass<bGTData, Chaos::FGeometryParticle, Chaos::FGeometryParticleHandle>::Result;
+
+	TBPVisitor(const FTransform& InWorldTM, ChaosInterface::FSQHitBuffer<THitType>& InHitBuffer,
 		const FQueryFilterData& InQueryFilterData, ICollisionQueryFilterCallbackBase& InQueryCallback, const QueryGeometryType& InQueryGeom, const FQueryDebugParams& InDebugParams)
 		: HalfExtents(InQueryGeom.BoundingBox().Extents() * 0.5)
 		, bAnyHit(false)
@@ -514,6 +520,18 @@ private:
 		Overlap
 	};
 
+	auto GetPayloadForThread(const TPayload& Payload)
+	{
+		if constexpr (bGTData)
+		{
+			return Payload.GetExternalGeometryParticle_ExternalThread();
+		}
+		else
+		{
+			return Payload.GetGeometryParticleHandle_PhysicsThread();
+		}
+	}
+
 	template <ESQType SQ>
 	bool Visit(const Chaos::TSpatialVisitorData<TPayload>& Instance, Chaos::FQueryFastData* CurData)
 	{
@@ -521,11 +539,12 @@ private:
 		TPayload Payload = Instance.Payload;
 		//todo: add a check to ensure hitbuffer matches SQ type
 		using namespace Chaos;
-		FGeometryParticle* GeometryParticle = Payload.GetExternalGeometryParticle_ExternalThread();
+		TGeometryType* GeometryParticle = GetPayloadForThread(Payload);
 		if (!GeometryParticle)
 		{
 			// This case handles particles created by the physics simulation without the main thread
 			// being made aware of their creation. We have a PT particle but no external particle
+			ensure(bGTData);
 			return true;
 		}
 		const FShapesArray& Shapes = GeometryParticle->ShapesArray();
@@ -585,15 +604,15 @@ void FChaosSQAccelerator::Raycast(const FVector& Start, const FVector& Dir, cons
 	RaycastImp(Start, Dir, DeltaMagnitude, HitBuffer, OutputFlags, QueryFilterData, QueryCallback, DebugParams);
 }
 
-template <typename QueryGeomType>
-void SweepHelper(const QueryGeomType& QueryGeom,const Chaos::ISpatialAcceleration<Chaos::FAccelerationStructureHandle, Chaos::FReal,3>& SpatialAcceleration,const FTransform& StartTM,const FVector& Dir,const float DeltaMagnitude,ChaosInterface::FSQHitBuffer<ChaosInterface::FSweepHit>& HitBuffer,EHitFlags OutputFlags,const FQueryFilterData& QueryFilterData,ICollisionQueryFilterCallbackBase& QueryCallback,const FQueryDebugParams& DebugParams)
+template <typename QueryGeomType, typename TSweepHit>
+void SweepHelper(const QueryGeomType& QueryGeom,const Chaos::ISpatialAcceleration<Chaos::FAccelerationStructureHandle, Chaos::FReal,3>& SpatialAcceleration,const FTransform& StartTM,const FVector& Dir,const float DeltaMagnitude,ChaosInterface::FSQHitBuffer<TSweepHit>& HitBuffer,EHitFlags OutputFlags,const FQueryFilterData& QueryFilterData,ICollisionQueryFilterCallbackBase& QueryCallback,const FQueryDebugParams& DebugParams)
 {
 	using namespace Chaos;
 	using namespace ChaosInterface;
 
 	const FAABB3 Bounds = QueryGeom.BoundingBox().TransformedAABB(StartTM);
 	const bool bSweepAsOverlap = DeltaMagnitude == 0;	//question: do we care about tiny sweeps?
-	TSQVisitor<QueryGeomType,FAccelerationStructureHandle,FSweepHit> SweepVisitor(StartTM,Dir,HitBuffer,OutputFlags,QueryFilterData,QueryCallback,QueryGeom,DebugParams);
+	TSQVisitor<QueryGeomType,FAccelerationStructureHandle,TSweepHit, std::is_same<TSweepHit, FSweepHit>::value> SweepVisitor(StartTM,Dir,HitBuffer,OutputFlags,QueryFilterData,QueryCallback,QueryGeom,DebugParams);
 
 	HitBuffer.IncFlushCount();
 
@@ -615,8 +634,13 @@ void FChaosSQAccelerator::Sweep(const Chaos::FImplicitObject& QueryGeom, const F
 	return Chaos::Utilities::CastHelper(QueryGeom, StartTM, [&](const auto& Downcast, const FTransform& StartFullTM) { return SweepHelper(Downcast, SpatialAcceleration, StartFullTM, Dir, DeltaMagnitude, HitBuffer, OutputFlags, QueryFilterData, QueryCallback, DebugParams); });
 }
 
-template <typename QueryGeomType>
-void OverlapHelper(const QueryGeomType& QueryGeom, const Chaos::ISpatialAcceleration<Chaos::FAccelerationStructureHandle, Chaos::FReal, 3>& SpatialAcceleration, const FTransform& GeomPose, ChaosInterface::FSQHitBuffer<ChaosInterface::FOverlapHit>& HitBuffer, const FQueryFilterData& QueryFilterData, ICollisionQueryFilterCallbackBase& QueryCallback, const FQueryDebugParams& DebugParams)
+void FChaosSQAccelerator::Sweep(const Chaos::FImplicitObject& QueryGeom, const FTransform& StartTM, const FVector& Dir, const float DeltaMagnitude, ChaosInterface::FSQHitBuffer<ChaosInterface::FPTSweepHit>& HitBuffer, EHitFlags OutputFlags, const FQueryFilterData& QueryFilterData, ICollisionQueryFilterCallbackBase& QueryCallback, const FQueryDebugParams& DebugParams) const
+{
+	return Chaos::Utilities::CastHelper(QueryGeom, StartTM, [&](const auto& Downcast, const FTransform& StartFullTM) { return SweepHelper(Downcast, SpatialAcceleration, StartFullTM, Dir, DeltaMagnitude, HitBuffer, OutputFlags, QueryFilterData, QueryCallback, DebugParams); });
+}
+
+template <typename QueryGeomType, typename TOverlapHit>
+void OverlapHelper(const QueryGeomType& QueryGeom, const Chaos::ISpatialAcceleration<Chaos::FAccelerationStructureHandle, Chaos::FReal, 3>& SpatialAcceleration, const FTransform& GeomPose, ChaosInterface::FSQHitBuffer<TOverlapHit>& HitBuffer, const FQueryFilterData& QueryFilterData, ICollisionQueryFilterCallbackBase& QueryCallback, const FQueryDebugParams& DebugParams)
 {
 	using namespace Chaos;
 	using namespace ChaosInterface;
@@ -629,14 +653,15 @@ void OverlapHelper(const QueryGeomType& QueryGeom, const Chaos::ISpatialAccelera
 #else
 	bool bSkipNarrowPhase = QueryFilterData.flags & FPhysicsQueryFlag::eSKIPNARROWPHASE;
 #endif
+	constexpr bool bGTData = std::is_same<TOverlapHit, FOverlapHit>::value;
 	if (bSkipNarrowPhase)
 	{
-		TBPVisitor<QueryGeomType, FAccelerationStructureHandle, FOverlapHit> OverlapVisitor(GeomPose, HitBuffer, QueryFilterData, QueryCallback, QueryGeom, DebugParams);
+		TBPVisitor<QueryGeomType, FAccelerationStructureHandle, TOverlapHit, bGTData> OverlapVisitor(GeomPose, HitBuffer, QueryFilterData, QueryCallback, QueryGeom, DebugParams);
 		SpatialAcceleration.Overlap(Bounds, OverlapVisitor);
 	}
 	else
 	{
-		TSQVisitor<QueryGeomType, FAccelerationStructureHandle, FOverlapHit> OverlapVisitor(GeomPose, HitBuffer, QueryFilterData, QueryCallback, QueryGeom, DebugParams);
+		TSQVisitor<QueryGeomType, FAccelerationStructureHandle, TOverlapHit, bGTData> OverlapVisitor(GeomPose, HitBuffer, QueryFilterData, QueryCallback, QueryGeom, DebugParams);
 		SpatialAcceleration.Overlap(Bounds, OverlapVisitor);
 	}
 	HitBuffer.DecFlushCount();
@@ -647,6 +672,11 @@ void FChaosSQAccelerator::Overlap(const Chaos::FImplicitObject& QueryGeom, const
 	return Chaos::Utilities::CastHelper(QueryGeom, GeomPose, [&](const auto& Downcast, const FTransform& GeomFullPose) { return OverlapHelper(Downcast, SpatialAcceleration, GeomFullPose, HitBuffer, QueryFilterData, QueryCallback, DebugParams); });
 }
 
+void FChaosSQAccelerator::Overlap(const Chaos::FImplicitObject& QueryGeom, const FTransform& GeomPose, ChaosInterface::FSQHitBuffer<ChaosInterface::FPTOverlapHit>& HitBuffer, const FQueryFilterData& QueryFilterData, ICollisionQueryFilterCallbackBase& QueryCallback, const FQueryDebugParams& DebugParams) const
+{
+	return Chaos::Utilities::CastHelper(QueryGeom, GeomPose, [&](const auto& Downcast, const FTransform& GeomFullPose) { return OverlapHelper(Downcast, SpatialAcceleration, GeomFullPose, HitBuffer, QueryFilterData, QueryCallback, DebugParams); });
+}
+
 #if WITH_PHYSX
 FChaosSQAcceleratorAdapter::FChaosSQAcceleratorAdapter(const Chaos::ISpatialAcceleration<Chaos::FAccelerationStructureHandle, Chaos::FReal, 3>& InSpatialAcceleration)
 	: ChaosSQAccelerator(InSpatialAcceleration)
diff --git a/Engine/Source/Runtime/PhysicsCore/Public/SQAccelerator.h b/Engine/Source/Runtime/PhysicsCore/Public/SQAccelerator.h
--- a/Engine/Source/Runtime/PhysicsCore/Public/SQAccelerator.h
+++ b/Engine/Source/Runtime/PhysicsCore/Public/SQAccelerator.h
@@ -46,7 +46,10 @@ public:
 	void Raycast(const FVector& Start, const FVector& Dir, const float DeltaMagnitude, ChaosInterface::FSQHitBuffer<ChaosInterface::FPTRaycastHit>& HitBuffer, EHitFlags OutputFlags, const FQueryFilterData& QueryFilterData, ICollisionQueryFilterCallbackBase& QueryCallback, const FQueryDebugParams& DebugParams = FQueryDebugParams()) const;
 
 	void Sweep(const Chaos::FImplicitObject& QueryGeom, const FTransform& StartTM, const FVector& Dir, const float DeltaMagnitude, ChaosInterface::FSQHitBuffer<ChaosInterface::FSweepHit>& HitBuffer, EHitFlags OutputFlags, const FQueryFilterData& QueryFilterData, ICollisionQueryFilterCallbackBase& QueryCallback, const FQueryDebugParams& DebugParams = FQueryDebugParams()) const;
+	void Sweep(const Chaos::FImplicitObject& QueryGeom, const FTransform& StartTM, const FVector& Dir, const float DeltaMagnitude, ChaosInterface::FSQHitBuffer<ChaosInterface::FPTSweepHit>& HitBuffer, EHitFlags OutputFlags, const FQueryFilterData& QueryFilterData, ICollisionQueryFilterCallbackBase& QueryCallback, const FQueryDebugParams& DebugParams = FQueryDebugParams()) const;
+
 	void Overlap(const Chaos::FImplicitObject& QueryGeom, const FTransform& GeomPose, ChaosInterface::FSQHitBuffer<ChaosInterface::FOverlapHit>& HitBuffer, const FQueryFilterData& QueryFilterData, ICollisionQueryFilterCallbackBase& QueryCallback, const FQueryDebugParams& DebugParams = FQueryDebugParams()) const;
+	void Overlap(const Chaos::FImplicitObject& QueryGeom, const FTransform& GeomPose, ChaosInterface::FSQHitBuffer<ChaosInterface::FPTOverlapHit>& HitBuffer, const FQueryFilterData& QueryFilterData, ICollisionQueryFilterCallbackBase& QueryCallback, const FQueryDebugParams& DebugParams = FQueryDebugParams()) const;
 
 private:
 	const Chaos::ISpatialAcceleration<Chaos::FAccelerationStructureHandle, Chaos::FReal, 3>& SpatialAcceleration;
