Fixed Chaos scene using a game thread only API for an external call which can come from any non-physics thread

diff --git a/Engine/Source/Runtime/PhysicsCore/Private/ChaosScene.cpp b/Engine/Source/Runtime/PhysicsCore/Private/ChaosScene.cpp
--- a/Engine/Source/Runtime/PhysicsCore/Private/ChaosScene.cpp
+++ b/Engine/Source/Runtime/PhysicsCore/Private/ChaosScene.cpp
@@ -83,7 +83,7 @@ FChaosScene::FChaosScene(
 	if(CVar_ApplyProjectSettings.GetValueOnAnyThread() != 0)
 	{
 		UPhysicsSettingsCore* Settings = UPhysicsSettingsCore::Get();
-		SceneSolver->EnqueueCommandImmediate([InSolver = SceneSolver, SolverConfigCopy = Settings->SolverOptions]()
+		SceneSolver->RegisterSimOneShotCallback([InSolver = SceneSolver, SolverConfigCopy = Settings->SolverOptions]()
 		{
 			InSolver->ApplyConfig(SolverConfigCopy);
 		});
