Chaos: Enable mask filters

Masks are now propagated to shapes correctly whereas previously the call to SetMaskFilter was a no-op and dropped the mask filter.

To match previous behaviour for mask filters (completely ignored for simulation - only ever considered for query) the mask filter check was removed for the sim filter and sim pre-filter.

diff --git a/Engine/Source/Runtime/Experimental/Chaos/Public/Chaos/Collision/CollisionFilter.h b/Engine/Source/Runtime/Experimental/Chaos/Public/Chaos/Collision/CollisionFilter.h
--- a/Engine/Source/Runtime/Experimental/Chaos/Public/Chaos/Collision/CollisionFilter.h
+++ b/Engine/Source/Runtime/Experimental/Chaos/Public/Chaos/Collision/CollisionFilter.h
@@ -92,11 +92,6 @@ namespace Chaos
 				const uint32 Filter0Channel = GetChaosCollisionChannelAndExtraFilter(Shape0->GetSimData().Word3, Filter0Mask);
 				const uint32 Filter1Channel = GetChaosCollisionChannelAndExtraFilter(Shape1->GetSimData().Word3, Filter1Mask);
 
-				if ((Filter0Mask & Filter1Mask) != 0)
-				{
-					return false;
-				}
-
 				const uint32 Filter1Bit = 1 << (Filter1Channel); // SIMDATA_TO_BITFIELD
 				uint32 const Filter0Bit = 1 << (Filter0Channel); // SIMDATA_TO_BITFIELD
 				return (Filter0Bit & Shape1->GetSimData().Word1) && (Filter1Bit & Shape0->GetSimData().Word1);
diff --git a/Engine/Source/Runtime/Experimental/Chaos/Public/Chaos/ParticleHandle.h b/Engine/Source/Runtime/Experimental/Chaos/Public/Chaos/ParticleHandle.h
--- a/Engine/Source/Runtime/Experimental/Chaos/Public/Chaos/ParticleHandle.h
+++ b/Engine/Source/Runtime/Experimental/Chaos/Public/Chaos/ParticleHandle.h
@@ -174,11 +174,6 @@ FORCEINLINE_DEBUGGABLE bool PrePreSimFilterImp(const FCollisionFilterData& SimFi
 	uint8  OtherMaskFilter;
 	const uint32 OtherChannel = GetChaosCollisionChannelAndExtraFilter(OtherSimFilterData.Word3, OtherMaskFilter);
 
-	if ((QuerierMaskFilter & OtherMaskFilter) != 0)
-	{
-		return true;
-	}
-
 	//uint32 const QuerierBit = ECC_TO_BITFIELD(QuerierChannel);
 	const uint32 QuerierBit = (1 << (QuerierChannel));
 	const uint32 OtherBit = (1 << (OtherChannel));
diff --git a/Engine/Source/Runtime/PhysicsCore/Private/ChaosEngineInterface.cpp b/Engine/Source/Runtime/PhysicsCore/Private/ChaosEngineInterface.cpp
--- a/Engine/Source/Runtime/PhysicsCore/Private/ChaosEngineInterface.cpp
+++ b/Engine/Source/Runtime/PhysicsCore/Private/ChaosEngineInterface.cpp
@@ -1521,6 +1521,26 @@ FPhysicsGeometryCollection_Chaos FChaosEngineInterface::GetGeometryCollection(co
 	return NewCollection;
 }
 
+void FChaosEngineInterface::SetMaskFilter(const FPhysicsShapeHandle& InShape, FMaskFilter InFilter)
+{
+	FCollisionFilterData SimFilter = GetSimulationFilter(InShape);
+	FCollisionFilterData QueryFilter = GetQueryFilter(InShape);
+
+	auto ApplyMask = [](uint32& Word3, FMaskFilter Mask)
+	{
+		// #CHAOSTODO - definitions for filter behavior are in the Engine module.
+		// Move all to PhysicsCore so we handle things in a safe way here.
+		static constexpr int32 LocalNumExtraBits = 6;
+		static_assert(LocalNumExtraBits <= 8, "Only up to 8 extra filter bits are supported.");
+		Word3 &= (0xFFFFFFFFu >> LocalNumExtraBits);	//we drop the top NumExtraFilterBits bits because that's where the new mask filter is going
+		Word3 |= uint32(Mask) << (32 - LocalNumExtraBits);
+	};
+	ApplyMask(SimFilter.Word3, InFilter);
+	ApplyMask(QueryFilter.Word3, InFilter);
+
+	SetSimulationFilter(InShape, SimFilter);
+	SetQueryFilter(InShape, QueryFilter);
+}
 
 FCollisionFilterData FChaosEngineInterface::GetSimulationFilter(const FPhysicsShapeReference_Chaos& InShape)
 {
diff --git a/Engine/Source/Runtime/PhysicsCore/Public/Chaos/ChaosEngineInterface.h b/Engine/Source/Runtime/PhysicsCore/Public/Chaos/ChaosEngineInterface.h
--- a/Engine/Source/Runtime/PhysicsCore/Public/Chaos/ChaosEngineInterface.h
+++ b/Engine/Source/Runtime/PhysicsCore/Public/Chaos/ChaosEngineInterface.h
@@ -534,7 +534,7 @@ public:
 
 	// @todo(mlentine): Which of these do we need to support?
 	// Set the mask filter of a shape, which is an extra level of filtering during collision detection / query for extra channels like "Blue Team" and "Red Team"
-	static void SetMaskFilter(const FPhysicsShapeHandle& InShape,FMaskFilter InFilter) {}
+	static void SetMaskFilter(const FPhysicsShapeHandle& InShape,FMaskFilter InFilter);
 	static void SetSimulationFilter(const FPhysicsShapeHandle& InShape,const FCollisionFilterData& InFilter);
 	static void SetQueryFilter(const FPhysicsShapeHandle& InShape,const FCollisionFilterData& InFilter);
 	static void SetIsSimulationShape(const FPhysicsShapeHandle& InShape,bool bIsSimShape);
