broad phase optimization:
Using bounding volume of rotated geometry instead of rotating AABB

diff --git a/Engine/Source/Runtime/PhysicsCore/Private/SQAccelerator.cpp b/Engine/Source/Runtime/PhysicsCore/Private/SQAccelerator.cpp
--- a/Engine/Source/Runtime/PhysicsCore/Private/SQAccelerator.cpp
+++ b/Engine/Source/Runtime/PhysicsCore/Private/SQAccelerator.cpp
@@ -257,7 +257,7 @@ private:
 		bool bContinue = true;
 
 		const FRigidTransform3 ActorTM(GeometryParticle->X(), GeometryParticle->R());
-		const TAABB<FReal, 3> QueryGeomWorldBounds = QueryGeom ? QueryGeom->BoundingBox().TransformedAABB(StartTM) : TAABB<FReal, 3>(-HalfExtents, HalfExtents);
+		const TAABB<FReal, 3> QueryGeomWorldBounds = QueryGeom ? QueryGeom->CalculateTransformedBounds(StartTM) : TAABB<FReal, 3>(-HalfExtents, HalfExtents);
 
 #if CHAOS_DEBUG_DRAW
 		bool bAllShapesIgnoredInPrefilter = true;
@@ -667,7 +667,7 @@ void SweepHelper(const QueryGeomType& QueryGeom,const Chaos::ISpatialAcceleratio
 	if(bSweepAsOverlap)
 	{
 		//fallback to overlap
-		SpatialAcceleration.Overlap(Bounds,SweepVisitor);
+		SpatialAcceleration.Overlap(Bounds, SweepVisitor);
 	} else
 	{
 		const FVector HalfExtents = Bounds.Extents() * 0.5f;
@@ -693,7 +693,7 @@ void OverlapHelper(const QueryGeomType& QueryGeom, const Chaos::ISpatialAccelera
 	using namespace Chaos;
 	using namespace ChaosInterface;
 
-	const FAABB3 Bounds = QueryGeom.BoundingBox().TransformedAABB(GeomPose);
+	const FAABB3 Bounds = QueryGeom.CalculateTransformedBounds(GeomPose);
 
 	HitBuffer.IncFlushCount();
 #if PHYSICS_INTERFACE_PHYSX
