Physics interface cleanup.

diff --git a/Engine/Source/Runtime/PhysicsCore/Private/PhysicsInitialization.cpp b/Engine/Source/Runtime/PhysicsCore/Private/PhysicsInitialization.cpp
--- a/Engine/Source/Runtime/PhysicsCore/Private/PhysicsInitialization.cpp
+++ b/Engine/Source/Runtime/PhysicsCore/Private/PhysicsInitialization.cpp
@@ -5,10 +5,7 @@
 #include "PhysicsPublicCore.h"
 #include "PhysXSupportCore.h"
 #include "Misc/CommandLine.h"
-#if PHYSICS_INTERFACE_PHYSX
-#include "IPhysXCookingModule.h"
-#include "IPhysXCooking.h"
-#endif
+
 #include "Modules/ModuleManager.h"
 #include "Misc/Paths.h"
 #include "Core/Public/HAL/IConsoleManager.h"
@@ -17,7 +14,6 @@
 #define APEX_STATICALLY_LINKED	0
 #endif
 
-
 // CVars
 TAutoConsoleVariable<float> CVarToleranceScaleLength(
 	TEXT("p.ToleranceScale_Length"),
@@ -42,237 +38,15 @@ bool InitGamePhysCore()
 	// If we're running with Chaos enabled, load its module
 	FModuleManager::Get().LoadModule("Chaos");
 
-#if WITH_ENGINE && WITH_CHAOS
+#if WITH_ENGINE
 	// Loading this without Chaos gives warning, as no module depends on it.
 	FModuleManager::Get().LoadModule("ChaosSolverEngine");
 #endif
 
-#if PHYSICS_INTERFACE_PHYSX
-	// Do nothing if SDK already exists
-	if (GPhysXFoundation != nullptr)
-	{
-		return true;
-	}
-
-	// Make sure 
-	if (!PhysDLLHelper::LoadPhysXModules(/*bLoadCookingModule=*/ false))
-	{
-		// This is fatal. We were not able to successfully load the physics modules
-		return false;
-	}
-
-	// Create Foundation
-	GPhysXAllocator = new FPhysXAllocator();
-	FPhysXErrorCallback* ErrorCallback = new FPhysXErrorCallback();
-
-	GPhysXFoundation = PxCreateFoundation(PX_FOUNDATION_VERSION, *GPhysXAllocator, *ErrorCallback);
-	check(GPhysXFoundation);
-
-#if PHYSX_MEMORY_STATS
-	// Want names of PhysX allocations
-	GPhysXFoundation->setReportAllocationNames(true);
-#endif
-
-	// Create profile manager
-	GPhysXVisualDebugger = PxCreatePvd(*GPhysXFoundation);
-	check(GPhysXVisualDebugger);
-
-	// Create Physics
-	PxTolerancesScale PScale;
-	PScale.length = CVarToleranceScaleLength.GetValueOnGameThread();
-	PScale.speed = CVarToleranceScaleSpeed.GetValueOnGameThread();
-
-	GPhysXSDK = PxCreatePhysics(PX_PHYSICS_VERSION, *GPhysXFoundation, PScale, false, GPhysXVisualDebugger);
-	check(GPhysXSDK);
-
-	FPhysxSharedData::Initialize();
-
-	// Init Extensions
-	PxInitExtensions(*GPhysXSDK, GPhysXVisualDebugger);
-
-	if (CVarUseUnifiedHeightfield.GetValueOnGameThread())
-	{
-		//Turn on PhysX 3.3 unified height field collision detection. 
-		//This approach shares the collision detection code between meshes and height fields such that height fields behave identically to the equivalent terrain created as a mesh. 
-		//This approach facilitates mixing the use of height fields and meshes in the application with no tangible difference in collision behavior between the two approaches except that 
-		//heightfield thickness is not supported for unified heightfields.
-		PxRegisterUnifiedHeightFields(*GPhysXSDK);
-	}
-	else
-	{
-		PxRegisterHeightFields(*GPhysXSDK);
-	}
-
-	if (FParse::Param(FCommandLine::Get(), TEXT("PVD")))
-	{
-		PvdConnect(TEXT("localhost"), true);
-	}
-
-	// Create Cooking
-	PxCooking* PhysXCooking = nullptr;
-	if (IPhysXCookingModule* Module = GetPhysXCookingModule())
-	{
-		PhysXCooking = Module->GetPhysXCooking()->GetCooking();
-	}
-
-#if WITH_APEX
-	check(PhysXCooking);	//APEX requires cooking
-
-	// Build the descriptor for the APEX SDK
-	apex::ApexSDKDesc ApexDesc;
-	ApexDesc.foundation = GPhysXFoundation;	// Pointer to the PxFoundation
-	ApexDesc.physXSDK = GPhysXSDK;	// Pointer to the PhysXSDK
-	ApexDesc.cooking = PhysXCooking;	// Pointer to the cooking library
-	ApexDesc.renderResourceManager = &GApexNullRenderResourceManager;	// We will not be using the APEX rendering API, so just use a dummy render resource manager
-	ApexDesc.resourceCallback = &GApexResourceCallback;	// The resource callback is how APEX asks the application to find assets when it needs them
-
-#if PLATFORM_MAC
-	FString DylibFolder = FPaths::EngineDir() / TEXT("Binaries/ThirdParty/PhysX3/");
-	ANSICHAR* DLLLoadPath = (ANSICHAR*)FMemory::Malloc(DylibFolder.Len() + 1);
-	FCStringAnsi::Strcpy(DLLLoadPath, DylibFolder.Len() + 1, TCHAR_TO_UTF8(*DylibFolder));
-	ApexDesc.dllLoadPath = DLLLoadPath;
-#endif
-
-	// Create the APEX SDK
-	apex::ApexCreateError ErrorCode;
-	GApexSDK = apex::CreateApexSDK(ApexDesc, &ErrorCode);
-	check(ErrorCode == APEX_CE_NO_ERROR);
-	check(GApexSDK);
-
-#if PLATFORM_MAC
-	FMemory::Free(DLLLoadPath);
-#endif
-
-#if UE_BUILD_SHIPPING
-	GApexSDK->setEnableApexStats(false);
-#endif
-
-#if APEX_STATICALLY_LINKED
-
-#if WITH_APEX_CLOTHING
-	instantiateModuleClothing();
-#endif
-
-#if WITH_APEX_LEGACY
-	instantiateModuleLegacy();
-#endif
-#endif
-
-	// 1 legacy module for all in APEX 1.3
-	// Load the only 1 legacy module
-#if WITH_APEX_LEGACY
-	GApexModuleLegacy = GApexSDK->createModule("Legacy", &ErrorCode);
-	check(GApexModuleLegacy);
-#endif // WITH_APEX_LEGACY
-
-#if WITH_APEX_CLOTHING
-	// Load APEX Clothing module
-	GApexModuleClothing = static_cast<apex::ModuleClothing*>(GApexSDK->createModule("Clothing"));
-	check(GApexModuleClothing);
-	// Set Clothing module parameters
-	NvParameterized::Interface* ModuleParams = GApexModuleClothing->getDefaultModuleDesc();
-
-	// Can be tuned for switching between more memory and more spikes.
-	NvParameterized::setParamU32(*ModuleParams, "maxUnusedPhysXResources", 5);
-
-	// If true, let fetch results tasks run longer than the fetchResults call. 
-	// Setting to true could not ensure same finish timing with Physx simulation phase
-	NvParameterized::setParamBool(*ModuleParams, "asyncFetchResults", false);
-
-	// ModuleParams contains the default module descriptor, which may be modified here before calling the module init function
-	GApexModuleClothing->init(*ModuleParams);
-#endif	//WITH_APEX_CLOTHING
-
-#endif // #if WITH_APEX
-
-#endif // WITH_PHYSX
-
 	return true;
 }
 
 void TermGamePhysCore()
 {
-#if PHYSICS_INTERFACE_PHYSX
-
-	FPhysxSharedData::Terminate();
-
-	// Do nothing if they were never initialized
-	if (GPhysXFoundation == nullptr)
-	{
-		return;
-	}
-
-#if WITH_APEX
-#if WITH_APEX_LEGACY
-	if (GApexModuleLegacy != nullptr)
-	{
-		GApexModuleLegacy->release();
-		GApexModuleLegacy = nullptr;
-	}
-#endif // WITH_APEX_LEGACY
-	if (GApexSDK != nullptr)
-	{
-		GApexSDK->release();
-		GApexSDK = nullptr;
-	}
-#endif	// #if WITH_APEX
-
-	//Remove all scenes still registered
-	if (GPhysXSDK != nullptr)
-	{
-		if (int32 NumScenes = GPhysXSDK->getNbScenes())
-		{
-			TArray<PxScene*> PScenes;
-			PScenes.AddUninitialized(NumScenes);
-			GPhysXSDK->getScenes(PScenes.GetData(), sizeof(PxScene*)* NumScenes);
-
-			for (PxScene* PScene : PScenes)
-			{
-				if (PScene)
-				{
-					PScene->release();
-				}
-			}
-		}
-	}
-
-	// Unload dependent modules
-	if(FModuleManager::Get().GetModule("PhysXVehicles"))
-	{
-		// Vehicles is actually in a plugin, but in order to shut down the foundation
-		// below - all dependents must release which requires us to shut this module
-		// down slightly early.
-		FModuleManager::Get().UnloadModule("PhysXVehicles", true);
-	}
-
-	if (IPhysXCookingModule* PhysXCookingModule = GetPhysXCookingModule( /*bforceLoad=*/false))
-	{
-		PhysXCookingModule->Terminate();
-	}
-
-	if (GPhysXSDK != nullptr)
-	{
-		PxCloseExtensions();
-	}
-
-	if (GPhysXSDK != nullptr)
-	{
-		GPhysXSDK->release();
-		GPhysXSDK = nullptr;
-	}
-
-	if(GPhysXVisualDebugger)
-	{
-		GPhysXVisualDebugger->release();
-		GPhysXVisualDebugger = nullptr;
-	}
-
-	GPhysXFoundation->release();
-	GPhysXFoundation = nullptr;
-
-	// @todo delete FPhysXAllocator
-	// @todo delete FPhysXOutputStream
 
-	PhysDLLHelper::UnloadPhysXModules();
-#endif // WITH_PHYSX
 }
diff --git a/Engine/Source/Runtime/PhysicsCore/Private/PhysicsInterfaceUtilsCore.cpp b/Engine/Source/Runtime/PhysicsCore/Private/PhysicsInterfaceUtilsCore.cpp
--- a/Engine/Source/Runtime/PhysicsCore/Private/PhysicsInterfaceUtilsCore.cpp
+++ b/Engine/Source/Runtime/PhysicsCore/Private/PhysicsInterfaceUtilsCore.cpp
@@ -6,7 +6,6 @@
 
 #include "Chaos/ParticleHandle.h"
 
-
 FCollisionFilterData C2UFilterData(const FChaosFilterData& PFilterData)
 {
 	FCollisionFilterData FilterData;
@@ -22,287 +21,18 @@ FChaosFilterData U2CFilterData(const FCollisionFilterData& FilterData)
 	return FChaosFilterData(FilterData.Word0, FilterData.Word1, FilterData.Word2, FilterData.Word3);
 }
 
-#if PHYSICS_INTERFACE_PHYSX
-FCollisionFilterData ToUnrealFilterData(const PxFilterData& FilterData)
-{
-	return P2UFilterData(FilterData);
-}
-#else
 FCollisionFilterData ToUnrealFilterData(const FChaosFilterData& FilterData)
 {
 	return C2UFilterData(FilterData);
 }
-#endif
-
-
-#if PHYSICS_INTERFACE_PHYSX
-
-FCollisionFilterData P2UFilterData(const PxFilterData& PFilterData)
-{
-	FCollisionFilterData FilterData;
-	FilterData.Word0 = PFilterData.word0;
-	FilterData.Word1 = PFilterData.word1;
-	FilterData.Word2 = PFilterData.word2;
-	FilterData.Word3 = PFilterData.word3;
-	return FilterData;
-}
-
-PxFilterData U2PFilterData(const FCollisionFilterData& FilterData)
-{
-	return PxFilterData(FilterData.Word0, FilterData.Word1, FilterData.Word2, FilterData.Word3);
-}
-
-
-PxShapeFlags BuildPhysXShapeFlags(FBodyCollisionFlags BodyCollisionFlags, bool bPhysicsStatic, bool bIsTriangleMesh)
-{
-	PxShapeFlags ShapeFlags;
-
-	// Only perform scene queries if enabled
-	ModifyShapeFlag<PxShapeFlag::eSCENE_QUERY_SHAPE>(ShapeFlags, BodyCollisionFlags.bEnableQueryCollision);
-
-	if (bIsTriangleMesh)
-	{
-		ModifyShapeFlag<PxShapeFlag::eSIMULATION_SHAPE>(ShapeFlags, BodyCollisionFlags.bEnableSimCollisionComplex);
-	}
-	else
-	{
-		ModifyShapeFlag<PxShapeFlag::eSIMULATION_SHAPE>(ShapeFlags, BodyCollisionFlags.bEnableSimCollisionSimple);
-	}
-
-	ModifyShapeFlag<PxShapeFlag::eVISUALIZATION>(ShapeFlags, true);
-
-	return ShapeFlags;
-}
-
-PxGeometryType::Enum U2PCollisionShapeType(ECollisionShapeType InUType)
-{
-	switch (InUType)
-	{
-	case ECollisionShapeType::Sphere:				return PxGeometryType::eSPHERE;
-	case ECollisionShapeType::Box:					return PxGeometryType::eBOX;
-	case ECollisionShapeType::Capsule:				return PxGeometryType::eCAPSULE;
-	case ECollisionShapeType::Convex:				return PxGeometryType::eCONVEXMESH;
-	case ECollisionShapeType::Trimesh:				return PxGeometryType::eTRIANGLEMESH;
-	case ECollisionShapeType::Heightfield:			return PxGeometryType::eHEIGHTFIELD;
-	}
-
-	return PxGeometryType::eINVALID;
-}
-
-ECollisionShapeType P2UCollisionShapeType(PxGeometryType::Enum InPType)
-{
-	switch (InPType)
-	{
-	case PxGeometryType::eSPHERE:					return ECollisionShapeType::Sphere;
-	case PxGeometryType::eBOX:						return ECollisionShapeType::Box;
-	case PxGeometryType::eCAPSULE:					return ECollisionShapeType::Capsule;
-	case PxGeometryType::eCONVEXMESH:				return ECollisionShapeType::Convex;
-	case PxGeometryType::eTRIANGLEMESH:				return ECollisionShapeType::Trimesh;
-	case PxGeometryType::eHEIGHTFIELD:				return ECollisionShapeType::Heightfield;
-	}
-
-	return ECollisionShapeType::None;
-}
-
-FPhysXShapeAdapter::FPhysXShapeAdapter(const FQuat& Rot, const FCollisionShape& CollisionShape)
-	: Rotation(FQuat::Identity)
-{
-	// Perform other kinds of zero-extent queries as zero-extent sphere queries
-	if ((CollisionShape.ShapeType != ECollisionShape::Sphere) && CollisionShape.IsNearlyZero())
-	{
-		GeometryHolder.storeAny(PxSphereGeometry(FCollisionShape::MinSphereRadius()));
-	}
-	else
-	{
-		switch (CollisionShape.ShapeType)
-		{
-		case ECollisionShape::Box:
-		{
-			PxVec3 BoxExtents = U2PVector(CollisionShape.GetBox());
-			BoxExtents.x = FMath::Max(BoxExtents.x, FCollisionShape::MinBoxExtent());
-			BoxExtents.y = FMath::Max(BoxExtents.y, FCollisionShape::MinBoxExtent());
-			BoxExtents.z = FMath::Max(BoxExtents.z, FCollisionShape::MinBoxExtent());
-
-			GeometryHolder.storeAny(PxBoxGeometry(BoxExtents));
-			Rotation = Rot;
-			break;
-		}
-		case ECollisionShape::Sphere:
-			GeometryHolder.storeAny(PxSphereGeometry(FMath::Max(CollisionShape.GetSphereRadius(), FCollisionShape::MinSphereRadius())));
-			break;
-		case ECollisionShape::Capsule:
-		{
-			const float CapsuleRadius = CollisionShape.GetCapsuleRadius();
-			const float CapsuleHalfHeight = CollisionShape.GetCapsuleHalfHeight();
-			if (CapsuleRadius < CapsuleHalfHeight)
-			{
-				GeometryHolder.storeAny(PxCapsuleGeometry(FMath::Max(CapsuleRadius, FCollisionShape::MinCapsuleRadius()), FMath::Max(CollisionShape.GetCapsuleAxisHalfLength(), FCollisionShape::MinCapsuleAxisHalfHeight())));
-				Rotation = ConvertToPhysXCapsuleRot(Rot);
-			}
-			else
-			{
-				// Use a sphere instead.
-				GeometryHolder.storeAny(PxSphereGeometry(FMath::Max(CapsuleRadius, FCollisionShape::MinSphereRadius())));
-			}
-			break;
-		}
-		default:
-			// invalid type
-			ensure(false);
-		}
-	}
-}
-
-static const PxQuat CapsuleRotator(0.f, 0.707106781f, 0.f, 0.707106781f);
-
-FQuat ConvertToPhysXCapsuleRot(const FQuat& GeomRot)
-{
-	// Rotation required because PhysX capsule points down X, we want it down Z
-	return P2UQuat(U2PQuat(GeomRot) * CapsuleRotator);	//todo(ocohen): avoid pointless conversion
-}
-
-FQuat ConvertToUECapsuleRot(const PxQuat & PGeomRot)
-{
-	return P2UQuat(PGeomRot * CapsuleRotator.getConjugate());
-}
-
-FQuat ConvertToUECapsuleRot(const FQuat & GeomRot)
-{
-	return GeomRot * P2UQuat(CapsuleRotator.getConjugate());
-}
-
-PxTransform ConvertToPhysXCapsulePose(const FTransform& GeomPose)
-{
-	PxTransform PFinalPose;
-
-	PFinalPose.p = U2PVector(GeomPose.GetTranslation());
-	// Rotation required because PhysX capsule points down X, we want it down Z
-	PFinalPose.q = U2PQuat(ConvertToPhysXCapsuleRot(GeomPose.GetRotation()));
-	return PFinalPose;
-}
-
-#if USE_SCENE_LOCK
-DEFINE_STAT(STAT_PhysSceneReadLock);
-DEFINE_STAT(STAT_PhysSceneWriteLock);
-#endif
-
-#endif
 
 template <typename THitLocation>
 uint32 FindFaceIndexImp(const THitLocation& PHit, const FVector& UnitDir)
 {
-#if PHYSICS_INTERFACE_PHYSX
-	PxConvexMeshGeometry convexGeom;
-	PxVec3 PUnitDir = U2PVector(UnitDir);
-	if(PHit.shape->getConvexMeshGeometry(convexGeom))
-	{
-		//PhysX has given us the most correct face. However, we actually want the most useful face which is the one with the most opposed normal within some radius.
-		//So for example, if we are sweeping against a corner we should take the plane that is most opposing, even if it's not the exact one we hit.
-		static const float FindFaceInRadius = 1.f; // tolerance to determine how far from the actual contact point we want to search.
-
-		PxTransform actorPose;
-		bool bFoundKinematicTarget = false;
-
-		if(PxRigidDynamic* rigidDynamic = PHit.actor->is<PxRigidDynamic>())
-		{
-			if(rigidDynamic->getRigidBodyFlags() & PxRigidBodyFlag::eUSE_KINEMATIC_TARGET_FOR_SCENE_QUERIES)
-			{
-				bFoundKinematicTarget = rigidDynamic->getKinematicTarget(actorPose);
-			}
-		}
-
-		if(!bFoundKinematicTarget)
-		{
-			actorPose = PHit.actor->getGlobalPose();
-		}
-
-		const PxTransform pose = actorPose * PHit.shape->getLocalPose();
-		const PxVec3 impactPos(PHit.position);
-		{
-			//This is copied directly from PxFindFace. However, we made some modifications in order to favor 'most opposing' faces.
-			static const PxReal gEpsilon = .01f;
-			PX_ASSERT(PUnitDir.isFinite());
-			PX_ASSERT(PUnitDir.isNormalized());
-			PX_ASSERT(impactPos.isFinite());
-			PX_ASSERT(pose.isFinite());
-
-			const PxVec3 impact = impactPos - PUnitDir * gEpsilon;
-
-			const PxVec3 localPoint = pose.transformInv(impact);
-			const PxVec3 localDir = pose.rotateInv(PUnitDir);
-
-			// Create shape to vertex scale transformation matrix
-			const PxMeshScale& meshScale = convexGeom.scale;
-			const PxMat33 rot(meshScale.rotation);
-			PxMat33 shape2VertexSkew = rot.getTranspose();
-			const PxMat33 diagonal = PxMat33::createDiagonal(PxVec3(1.0f / meshScale.scale.x, 1.0f / meshScale.scale.y, 1.0f / meshScale.scale.z));
-			shape2VertexSkew = shape2VertexSkew * diagonal;
-			shape2VertexSkew = shape2VertexSkew * rot;
-
-			const PxU32 nbPolys = convexGeom.convexMesh->getNbPolygons();
-			// BEGIN EPIC MODIFICATION Improved selection of 'most opposing' face
-			bool bMinIndexValid = false;
-			PxU32 minIndex = 0;
-			PxReal maxD = -PX_MAX_REAL;
-			PxU32 maxDIndex = 0;
-			PxReal minNormalDot = PX_MAX_REAL;
-
-			for(PxU32 j = 0; j < nbPolys; j++)
-			{
-				PxHullPolygon hullPolygon;
-				convexGeom.convexMesh->getPolygonData(j, hullPolygon);
-
-				// transform hull plane into shape space
-				PxPlane plane;
-				const PxVec3 tmp = shape2VertexSkew.transformTranspose(PxVec3(hullPolygon.mPlane[0], hullPolygon.mPlane[1], hullPolygon.mPlane[2]));
-				const PxReal denom = 1.0f / tmp.magnitude();
-				plane.n = tmp * denom;
-				plane.d = hullPolygon.mPlane[3] * denom;
-
-				PxReal d = plane.distance(localPoint);
-				// Track plane that impact point is furthest point (will be out fallback normal)
-				if(d > maxD)
-				{
-					maxDIndex = j;
-					maxD = d;
-				}
-
-				//Because we are searching against a convex hull, we will never get multiple faces that are both in front of the contact point _and_ have an opposing normal (except the face we hit).
-				//However, we may have just missed a plane which is now "behind" the contact point while still being inside the radius
-				if(d < -FindFaceInRadius)
-					continue;
-
-				// Calculate direction dot plane normal
-				const PxReal normalDot = plane.n.dot(localDir);
-				// If this is more opposing than our current 'most opposing' normal, update 'most opposing'
-				if(normalDot < minNormalDot)
-				{
-					minIndex = j;
-					bMinIndexValid = true;
-					minNormalDot = normalDot;
-				}
-			}
-
-			// If we found at least one face that we are considered 'on', use best normal
-			if(bMinIndexValid)
-			{
-				return minIndex;
-			}
-			// Fallback is the face that we are most in front of
-			else
-			{
-				return maxDIndex;
-			}
-		}
-	}
-
-	return PHit.faceIndex;	//If no custom logic just return whatever face index they initially had
-#else
 	const FTransform WorldTM(PHit.Actor->R(), PHit.Actor->X());
 	const FVector LocalPosition = WorldTM.InverseTransformPositionNoScale(PHit.WorldPosition);
 	const FVector LocalUnitDir = WorldTM.InverseTransformVectorNoScale(UnitDir);
 	return PHit.Shape->GetGeometry()->FindMostOpposingFace(LocalPosition, LocalUnitDir, PHit.FaceIndex, 1);	//todo:this number matches the one above, but is it right?
-#endif
 }
 
 uint32 FindFaceIndex(const FHitLocation& Hit, const FVector& UnitDir)
diff --git a/Engine/Source/Runtime/PhysicsCore/Private/SQAccelerator.cpp b/Engine/Source/Runtime/PhysicsCore/Private/SQAccelerator.cpp
--- a/Engine/Source/Runtime/PhysicsCore/Private/SQAccelerator.cpp
+++ b/Engine/Source/Runtime/PhysicsCore/Private/SQAccelerator.cpp
@@ -4,14 +4,7 @@
 #include "CollisionQueryFilterCallbackCore.h"
 #include "PhysicsCore/Public/PhysicsInterfaceUtilsCore.h"
 
-#if PHYSICS_INTERFACE_PHYSX
-#include "SceneQueryPhysXImp.h"	//todo: use nice platform wrapper
-#include "PhysXInterfaceWrapperCore.h"
-#endif
-
-#if WITH_CHAOS
 #include "SceneQueryChaosImp.h"
-#endif
 
 #include "ChaosInterfaceWrapperCore.h"
 #include "Chaos/CastingUtilities.h"
@@ -112,18 +105,10 @@ struct TSQVisitor : public Chaos::ISpatialVisitor<TPayload, Chaos::FReal>
 		, DebugParams(InDebugParams)
 		, HitBuffer(InHitBuffer)
 		, QueryFilterData(InQueryFilterData)
-#if PHYSICS_INTERFACE_PHYSX
-		, QueryFilterDataConcrete(P2UFilterData(QueryFilterData.data))
-#else
 		, QueryFilterDataConcrete(C2UFilterData(QueryFilterData.data))
-#endif
 		, QueryCallback(InQueryCallback)
 	{
-#if PHYSICS_INTERFACE_PHYSX
-		bAnyHit = QueryFilterData.flags & PxQueryFlag::eANY_HIT;
-#else
 		bAnyHit = QueryFilterData.flags & FChaosQueryFlag::eANY_HIT;
-#endif
 	}
 
 	TSQVisitor(const FTransform& InStartTM, const FVector& InDir, ChaosInterface::FSQHitBuffer<THitType>& InHitBuffer, EHitFlags InOutputFlags,
@@ -136,21 +121,12 @@ struct TSQVisitor : public Chaos::ISpatialVisitor<TPayload, Chaos::FReal>
 		, HitFaceNormal(Chaos::FVec3::ZeroVector)
 		, HitBuffer(InHitBuffer)
 		, QueryFilterData(InQueryFilterData)
-#if PHYSICS_INTERFACE_PHYSX
-		, QueryFilterDataConcrete(P2UFilterData(QueryFilterData.data))
-#else
 		, QueryFilterDataConcrete(C2UFilterData(QueryFilterData.data))
-#endif
 		, QueryGeom(&InQueryGeom)
 		, QueryCallback(InQueryCallback)
 		, StartTM(InStartTM)
 	{
-#if PHYSICS_INTERFACE_PHYSX
-		bAnyHit = QueryFilterData.flags & PxQueryFlag::eANY_HIT;
-#else
 		bAnyHit = QueryFilterData.flags & FChaosQueryFlag::eANY_HIT;
-#endif
-
 		//todo: check THitType is sweep
 	}
 
@@ -162,21 +138,12 @@ struct TSQVisitor : public Chaos::ISpatialVisitor<TPayload, Chaos::FReal>
 		, HitFaceNormal(Chaos::FVec3::ZeroVector)
 		, HitBuffer(InHitBuffer)
 		, QueryFilterData(InQueryFilterData)
-#if PHYSICS_INTERFACE_PHYSX
-		, QueryFilterDataConcrete(P2UFilterData(QueryFilterData.data))
-#else
 		, QueryFilterDataConcrete(C2UFilterData(QueryFilterData.data))
-#endif
 		, QueryGeom(&InQueryGeom)
 		, QueryCallback(InQueryCallback)
 		, StartTM(InWorldTM)
 	{
-#if PHYSICS_INTERFACE_PHYSX
-		bAnyHit = QueryFilterData.flags & PxQueryFlag::eANY_HIT;
-#else
 		bAnyHit = QueryFilterData.flags & FChaosQueryFlag::eANY_HIT;
-#endif
-
 		//todo: check THitType is overlap
 	}
 
@@ -304,12 +271,7 @@ private:
 			}
 
 			//TODO: use gt particles directly
-#if PHYSICS_INTERFACE_PHYSX
-			ECollisionQueryHitType HitType = QueryFilterData.flags & PxQueryFlag::ePREFILTER ? QueryCallback.PreFilter(QueryFilterDataConcrete, *Shape, *GeometryParticle) : ECollisionQueryHitType::Block;
-#else
 			ECollisionQueryHitType HitType = QueryFilterData.flags & FChaosQueryFlag::ePREFILTER ? QueryCallback.PreFilter(QueryFilterDataConcrete, *Shape, *GeometryParticle) : ECollisionQueryHitType::Block;
-#endif
-
 
 			if (HitType != ECollisionQueryHitType::None)
 			{
@@ -412,11 +374,7 @@ private:
 					{
 						FillHitHelper(Hit, Distance, WorldPosition, WorldNormal, FaceIdx, bComputeMTD);
 
-#if PHYSICS_INTERFACE_PHYSX
-						HitType = QueryFilterData.flags & PxQueryFlag::ePOSTFILTER ? QueryCallback.PostFilter(QueryFilterDataConcrete, Hit) : HitType;
-#else
 						HitType = QueryFilterData.flags & FChaosQueryFlag::ePOSTFILTER ? QueryCallback.PostFilter(QueryFilterDataConcrete, Hit) : HitType;
-#endif
 
 						if(HitType != ECollisionQueryHitType::None)
 						{
@@ -449,7 +407,7 @@ private:
 			}
 		}
 
-#if CHAOS_DEBUG_DRAW && WITH_CHAOS
+#if CHAOS_DEBUG_DRAW
 		if (DebugParams.IsDebugQuery() && ChaosSQDrawDebugVisitorQueries)
 		{
 			DebugDraw<SQ>(Instance, CurData, bAllShapesIgnoredInPrefilter, bHitBufferIncreased);
@@ -483,12 +441,10 @@ private:
 			Chaos::FDebugDrawQueue::GetInstance().DrawDebugBox(Instance.Bounds.Center(), Instance.Bounds.Extents(), FQuat::Identity, bHit ? FColor(100, 50, 50) : FColor(50, 100, 50), false, -1.f, 0, 0.f);
 		}
 
-#if WITH_CHAOS
 		if (!bPrefiltered)
 		{
 			DebugDrawPayload(Instance.Payload, DebugParams.bExternalQuery, bHit);
 		}
-#endif
 	}
 #endif
 
@@ -696,11 +652,9 @@ void OverlapHelper(const QueryGeomType& QueryGeom, const Chaos::ISpatialAccelera
 	const FAABB3 Bounds = QueryGeom.CalculateTransformedBounds(GeomPose);
 
 	HitBuffer.IncFlushCount();
-#if PHYSICS_INTERFACE_PHYSX
-	bool bSkipNarrowPhase = false; // Flag doesn't exist on PhysX type.
-#else
+
 	bool bSkipNarrowPhase = QueryFilterData.flags & FPhysicsQueryFlag::eSKIPNARROWPHASE;
-#endif
+
 	constexpr bool bGTData = std::is_same<TOverlapHit, FOverlapHit>::value;
 	if (bSkipNarrowPhase)
 	{
@@ -724,70 +678,3 @@ void FChaosSQAccelerator::Overlap(const Chaos::FImplicitObject& QueryGeom, const
 {
 	return Chaos::Utilities::CastHelper(QueryGeom, GeomPose, [&](const auto& Downcast, const FTransform& GeomFullPose) { return OverlapHelper(Downcast, SpatialAcceleration, GeomFullPose, HitBuffer, QueryFilterData, QueryCallback, DebugParams); });
 }
-
-#if WITH_PHYSX
-FChaosSQAcceleratorAdapter::FChaosSQAcceleratorAdapter(const Chaos::ISpatialAcceleration<Chaos::FAccelerationStructureHandle, Chaos::FReal, 3>& InSpatialAcceleration)
-	: ChaosSQAccelerator(InSpatialAcceleration)
-{
-}
-
-void FChaosSQAcceleratorAdapter::Raycast(const FVector& Start, const FVector& Dir, const float DeltaMagnitude, FPhysicsHitCallback<FHitRaycast>& HitBuffer, EHitFlags OutputFlags, const FQueryFilterData& QueryFilterData, ICollisionQueryFilterCallbackBase& QueryCallback) const
-{
-	check(false);
-}
-
-void FChaosSQAcceleratorAdapter::Sweep(const FPhysicsGeometry& QueryGeom, const FTransform& StartTM, const FVector& Dir, const float DeltaMagnitude, FPhysicsHitCallback<FHitSweep>& HitBuffer, EHitFlags OutputFlags, const FQueryFilterData& QueryFilterData, ICollisionQueryFilterCallbackBase& QueryCallback) const
-{
-	check(false);
-}
-
-void FChaosSQAcceleratorAdapter::Overlap(const FPhysicsGeometry& QueryGeom, const FTransform& GeomPose, FPhysicsHitCallback<FHitOverlap>& HitBuffer, const FQueryFilterData& QueryFilterData, ICollisionQueryFilterCallbackBase& QueryCallback) const
-{
-	check(false);
-}
-#endif
-
-
-#if WITH_PHYSX && !WITH_CHAOS
-
-FPhysXSQAccelerator::FPhysXSQAccelerator()
-	: Scene(nullptr)
-{
-
-}
-
-FPhysXSQAccelerator::FPhysXSQAccelerator(physx::PxScene* InScene)
-	: Scene(InScene)
-{
-
-}
-
-void FPhysXSQAccelerator::Raycast(const FVector& Start, const FVector& Dir, const float DeltaMagnitude, FPhysicsHitCallback<FHitRaycast>& HitBuffer, EHitFlags OutputFlags, const FQueryFilterData& QueryFilterData, ICollisionQueryFilterCallbackBase& QueryCallback) const
-{
-	check(Scene);
-
-	FPhysicsRaycastInputAdapater Inputs(Start, Dir, OutputFlags);
-	Scene->raycast(Inputs.Start, Inputs.Dir, DeltaMagnitude, HitBuffer, Inputs.OutputFlags, QueryFilterData, &QueryCallback);
-}
-
-void FPhysXSQAccelerator::Sweep(const FPhysicsGeometry& QueryGeom, const FTransform& StartTM, const FVector& Dir, const float DeltaMagnitude, FPhysicsHitCallback<FHitSweep>& HitBuffer, EHitFlags OutputFlags, const FQueryFilterData& QueryFilterData, ICollisionQueryFilterCallbackBase& QueryCallback) const
-{
-	check(Scene);
-
-	FPhysicsSweepInputAdapater Inputs(StartTM, Dir, OutputFlags);
-	Scene->sweep(QueryGeom, Inputs.StartTM, Inputs.Dir, DeltaMagnitude, HitBuffer, Inputs.OutputFlags, QueryFilterData, &QueryCallback);
-}
-
-void FPhysXSQAccelerator::Overlap(const FPhysicsGeometry& QueryGeom, const FTransform& GeomPose, FPhysicsHitCallback<FHitOverlap>& HitBuffer, const FQueryFilterData& QueryFilterData, ICollisionQueryFilterCallbackBase& QueryCallback) const
-{
-	check(Scene);
-
-	FPhysicsOverlapInputAdapater Inputs(GeomPose);
-	Scene->overlap(QueryGeom, Inputs.GeomPose, HitBuffer, QueryFilterData, &QueryCallback);
-}
-
-void FPhysXSQAccelerator::SetScene(physx::PxScene* InScene)
-{
-	Scene = InScene;
-}
-#endif
diff --git a/Engine/Source/Runtime/PhysicsCore/Private/SQCapture.cpp b/Engine/Source/Runtime/PhysicsCore/Private/SQCapture.cpp
--- a/Engine/Source/Runtime/PhysicsCore/Private/SQCapture.cpp
+++ b/Engine/Source/Runtime/PhysicsCore/Private/SQCapture.cpp
@@ -10,10 +10,6 @@
 #include "PhysicsPublicCore.h"
 #include "PhysTestSerializer.h" 
 
-#if PHYSICS_INTERFACE_PHYSX
-#include "PhysXToChaosUtil.h"
-#endif
-
 class FSQCaptureFilterCallback : public ICollisionQueryFilterCallbackBase
 {
 public:
@@ -25,13 +21,6 @@ public:
 	virtual ECollisionQueryHitType PostFilter(const FCollisionFilterData& FilterData, const ChaosInterface::FPTQueryHit& Hit) override { ensure(false);  return ECollisionQueryHitType::Touch; }
 	virtual ECollisionQueryHitType PreFilter(const FCollisionFilterData& FilterData, const Chaos::FPerShapeData& Shape, const Chaos::FGeometryParticleHandle& Actor) override { ensure(false);  return ECollisionQueryHitType::Touch; }
 
-#if PHYSICS_INTERFACE_PHYSX
-	virtual ECollisionQueryHitType PostFilter(const FCollisionFilterData& FilterData, const physx::PxQueryHit& Hit) override { return ECollisionQueryHitType::Touch; }
-	virtual ECollisionQueryHitType PreFilter(const FCollisionFilterData& FilterData, const physx::PxShape& Shape, physx::PxRigidActor& Actor) override { return Capture.GetFilterResult(&Shape, &Actor); }
-	virtual PxQueryHitType::Enum preFilter(const PxFilterData& filterData, const PxShape* shape, const PxRigidActor* actor, PxHitFlags& queryFlags) override { return U2PCollisionQueryHitType(Capture.GetFilterResult(shape, actor)); }
-	virtual PxQueryHitType::Enum postFilter(const PxFilterData& filterData, const PxQueryHit& hit) override { return PxQueryHitType::eTOUCH; }
-#endif
-
 private:
 	const FSQCapture& Capture;
 };
@@ -50,136 +39,6 @@ FSQCapture::~FSQCapture()
 {
 }
 
-#if PHYSICS_INTERFACE_PHYSX
-
-PxActor* FSQCapture::GetTransientActor(PxActor* Actor) const
-{
-	PxActor* TransientActor = NonTransientToTransientActors.FindRef(Actor);
-	return TransientActor ? TransientActor : Actor;
-}
-
-PxShape* FSQCapture::GetTransientShape(PxShape* Shape) const
-{
-	PxShape* TransientShape = NonTransientToTransientShapes.FindRef(Shape);
-	return TransientShape ? TransientShape : Shape;
-}
-
-void FSQCapture::SerializePhysXHitType(FArchive& Ar, PxOverlapHit& Hit)
-{
-	uint64 Actor = (PxSerialObjectId)GetTransientActor(Hit.actor);
-	uint64 Shape = (PxSerialObjectId)GetTransientShape(Hit.shape);
-
-	Ar << Actor << Shape << Hit.faceIndex;
-	Hit.actor = (PxRigidActor*)Actor;	//todo: 32bit, should be ok as long as you never load a 64bit capture on a 32bit machine. If compiler error just add some checks and do conversion manually
-	Hit.shape = (PxShape*)Shape;
-}
-
-template <typename T>
-void FSQCapture::SerializePhysXHitType(FArchive& Ar, T& Hit)
-{
-	uint64 Actor = (PxSerialObjectId)GetTransientActor(Hit.actor);
-	uint64 Shape = (PxSerialObjectId)GetTransientShape(Hit.shape);
-	FVector Position = P2UVector(Hit.position);
-	FVector Normal = P2UVector(Hit.normal);
-	uint16 HitFlags = Hit.flags;
-
-	Ar << Actor << Shape << Hit.faceIndex << HitFlags << Position << Normal << Hit.distance;
-	Hit.actor = (PxRigidActor*)Actor;	//todo: 32bit, should be ok as long as you never load a 64bit capture on a 32bit machine. If compiler error just add some checks and do conversion manually
-	Hit.shape = (PxShape*)Shape;
-	Hit.position = U2PVector(Position);
-	Hit.normal = U2PVector(Normal);
-	Hit.flags = (PxHitFlags)HitFlags;
-}
-
-template <typename THit>
-void FSQCapture::SerializePhysXBuffers(FArchive& Ar, int32 Version, PhysXInterface::FDynamicHitBuffer<THit>& PhysXBuffer)
-{
-	Ar << PhysXBuffer.hasBlock;
-	if (PhysXBuffer.hasBlock)
-	{
-		SerializePhysXHitType(Ar, PhysXBuffer.block);
-	}
-
-	TArray<THit> TmpHits;
-	THit* Hits;
-	int32 NumHits;
-	if (Version < 1)
-	{
-		Ar << PhysXBuffer.maxNbTouches << PhysXBuffer.nbTouches;
-		TmpHits.AddDefaulted(PhysXBuffer.nbTouches);
-		Hits = TmpHits.GetData();
-		NumHits = TmpHits.Num();
-	}
-	else
-	{
-		NumHits = PhysXBuffer.GetNumHits();
-		Ar << NumHits;
-		if (Ar.IsLoading())
-		{
-			TmpHits.AddDefaulted(NumHits);
-			PhysXBuffer.processTouches(TmpHits.GetData(), NumHits);
-		}
-
-		Hits = PhysXBuffer.GetHits();
-	}
-
-	for (int32 Idx = 0; Idx < NumHits; ++Idx)
-	{
-		SerializePhysXHitType(Ar, Hits[Idx]);
-	}
-}
-
-void FSQCapture::SerializeActorToShapeHitsArray(FArchive& Ar)
-{
-	int32 NumActors = PxActorToShapeHitsArray.Num();
-	Ar << NumActors;
-	if (Ar.IsLoading())
-	{
-		for (int32 ActorIdx = 0; ActorIdx < NumActors; ++ActorIdx)
-		{
-			uint64 Actor;
-			Ar << Actor;
-			int32 NumShapes;
-			Ar << NumShapes;
-
-			TArray<TPair<PxShape*, ECollisionQueryHitType>> Pairs;
-			for (int32 ShapeIdx = 0; ShapeIdx < NumShapes; ++ShapeIdx)
-			{
-				uint64 Shape;
-				Ar << Shape;
-				ECollisionQueryHitType HitType;
-				Ar << HitType;
-				PxShape* ShapePtr = static_cast<PxShape*>(PhysSerializer.FindObject(Shape));
-				check(ShapePtr);
-				Pairs.Emplace(ShapePtr, HitType);
-			}
-			PxActor* ActorPtr = static_cast<PxActor*>(PhysSerializer.FindObject(Actor));
-			check(ActorPtr);
-			PxActorToShapeHitsArray.Add(ActorPtr, Pairs);
-		}
-	}
-	else if (Ar.IsSaving())
-	{
-		for (auto& Itr : PxActorToShapeHitsArray)
-		{
-			uint64 Actor = (PxSerialObjectId)NonTransientToTransientActors.FindChecked(Itr.Key);
-			Ar << Actor;
-			int32 NumShapes = Itr.Value.Num();
-			Ar << NumShapes;
-
-			for (auto& Pair : Itr.Value)
-			{
-				uint64 Shape = (PxSerialObjectId)NonTransientToTransientShapes.FindChecked(Pair.Key);
-				Ar << Shape;
-				Ar << Pair.Value;
-
-			}
-		}
-	}
-}
-
-#endif
-
 void SerializeQueryFilterData(FArchive& Ar, FQueryFilterData& QueryFilterData)
 {
 	Ar << QueryFilterData.data.word0;
@@ -188,11 +47,7 @@ void SerializeQueryFilterData(FArchive& Ar, FQueryFilterData& QueryFilterData)
 	Ar << QueryFilterData.data.word3;
 	uint16 Flags = QueryFilterData.flags;
 	Ar << Flags;
-#if PHYSICS_INTERFACE_PHYSX
-	QueryFilterData.flags = (PxQueryFlags)Flags;
-#else
 	QueryFilterData.flags = (FChaosQueryFlags)Flags;
-#endif
 	Ar << QueryFilterData.clientId;
 }
 
@@ -300,40 +155,19 @@ void FSQCapture::Serialize(Chaos::FChaosArchive& Ar)
 		Ar << StartPoint;
 	}
 
-#if WITH_PHYSX
 	if (bDiskDataIsChaos == false)
 	{
-#if PHYSICS_INTERFACE_PHYSX
-		SerializePhysXBuffers(Ar, Version, PhysXSweepBuffer);
-
-		if (Version >= 1)
-		{
-			SerializePhysXBuffers(Ar, Version, PhysXRaycastBuffer);
-			SerializePhysXBuffers(Ar, Version, PhysXOverlapBuffer);
-
-			SerializeActorToShapeHitsArray(Ar);
-		}
-
-		if (Ar.IsLoading())
-		{
-			CreatePhysXData();
-		}
-#endif
 		if(Version >= 1)
 		{
 			SerializeQueryFilterData(Ar, QueryFilterData);
 		}
 	}
-#endif
 
-#if WITH_CHAOS
 	if (bDiskDataIsChaos)
 	{
-#if WITH_CHAOS
 		SerializeChaosBuffers(Ar, Version, ChaosSweepBuffer);
 		SerializeChaosBuffers(Ar, Version, ChaosRaycastBuffer);
 		SerializeChaosBuffers(Ar, Version, ChaosOverlapBuffer);
-#endif // WITH_CHAOS
 
 		SerializeChaosActorToShapeHitsArray(Ar);
 		SerializeQueryFilterData(Ar, QueryFilterData);
@@ -344,14 +178,9 @@ void FSQCapture::Serialize(Chaos::FChaosArchive& Ar)
 			ChaosGeometry = SerializableChaosGeometry.Get();
 		}
 	}
-#endif
 
 	if (Ar.IsLoading())
 	{
-#if 0
-		CreateChaosDataFromPhysX();
-#endif // WITH_PHYSX
-
 		if (bDiskDataIsChaos)
 		{
 			FilterCallback = MakeUnique<FSQCaptureFilterCallback>(*this);
@@ -380,15 +209,11 @@ void FSQCapture::StartCaptureChaosSweep(const Chaos::FPBDRigidsEvolution& Evolut
 
 void FSQCapture::EndCaptureChaosSweep(const ChaosInterface::FSQHitBuffer<ChaosInterface::FSweepHit>& Results)
 {
-#if WITH_CHAOS
 	if (IsInGameThread())
 	{
 		check(SQType == ESQType::Sweep);
-#if WITH_CHAOS
 		ChaosSweepBuffer = Results;
-#endif // WITH_CHAOS
 	}
-#endif
 }
 
 void FSQCapture::StartCaptureChaosRaycast(const Chaos::FPBDRigidsEvolution& Evolution, const FVector& InStartPoint, const FVector& InDir, float InDeltaMag, FHitFlags InOutputFlags, const FQueryFilterData& QueryFilter, const FCollisionFilterData& FilterData, ICollisionQueryFilterCallbackBase& Callback)
@@ -410,15 +235,11 @@ void FSQCapture::StartCaptureChaosRaycast(const Chaos::FPBDRigidsEvolution& Evol
 
 void FSQCapture::EndCaptureChaosRaycast(const ChaosInterface::FSQHitBuffer<ChaosInterface::FRaycastHit>& Results)
 {
-#if WITH_CHAOS
 	if (IsInGameThread())
 	{
 		check(SQType == ESQType::Raycast);
-#if WITH_CHAOS
 		ChaosRaycastBuffer = Results;
-#endif // WITH_CHAOS
 	}
-#endif
 }
 
 void FSQCapture::StartCaptureChaosOverlap(const Chaos::FPBDRigidsEvolution& Evolution, const Chaos::FImplicitObject& InQueryGeom, const FTransform& InStartTM, const FQueryFilterData& QueryFilter, const FCollisionFilterData& FilterData, ICollisionQueryFilterCallbackBase& Callback)
@@ -439,15 +260,11 @@ void FSQCapture::StartCaptureChaosOverlap(const Chaos::FPBDRigidsEvolution& Evol
 
 void FSQCapture::EndCaptureChaosOverlap(const ChaosInterface::FSQHitBuffer<ChaosInterface::FOverlapHit>& Results)
 {
-#if WITH_CHAOS
 	if (IsInGameThread())
 	{
 		check(SQType == ESQType::Overlap);
-#if WITH_CHAOS
 		ChaosOverlapBuffer = Results;
-#endif // WITH_CHAOS
 	}
-#endif
 }
 
 void FSQCapture::CaptureChaosFilterResults(const Chaos::FPBDRigidsEvolution& TransientEvolution, const FCollisionFilterData& FilterData, ICollisionQueryFilterCallbackBase& Callback)
@@ -495,356 +312,3 @@ ECollisionQueryHitType FSQCapture::GetFilterResult(const Chaos::FPerShapeData* S
 {
 	return GetFilterResultHelper(Shape, Actor, ChaosActorToShapeHitsArray);
 }
-
-
-#if PHYSICS_INTERFACE_PHYSX
-void FSQCapture::StartCapturePhysXSweep(const PxScene& Scene, const PxGeometry& InQueryGeom, const FTransform& InStartTM, const FVector& InDir, float InDeltaMag, FHitFlags InOutputFlags, const FQueryFilterData& QueryFilter, const FCollisionFilterData& FilterData, ICollisionQueryFilterCallbackBase& Callback)
-{
-	if (IsInGameThread())
-	{
-		bDiskDataIsChaos = false;
-		CapturePhysXFilterResults(Scene, FilterData, Callback);
-		//copy data
-		PhysXGeometry.storeAny(InQueryGeom);
-		StartTM = InStartTM;
-		Dir = InDir;
-		DeltaMag = InDeltaMag;
-		OutputFlags = InOutputFlags;
-		QueryFilterData = QueryFilter;
-
-		SQType = ESQType::Sweep;
-		SetPhysXGeometryData(InQueryGeom);
-	}
-}
-
-void FSQCapture::StartCapturePhysXRaycast(const PxScene& Scene, const FVector& InStartPoint, const FVector& InDir, float InDeltaMag, FHitFlags InOutputFlags, const FQueryFilterData& QueryFilter, const FCollisionFilterData& FilterData, ICollisionQueryFilterCallbackBase& Callback)
-{
-	if (IsInGameThread())
-	{
-		bDiskDataIsChaos = false;
-		CapturePhysXFilterResults(Scene, FilterData, Callback);
-		//copy data
-		StartPoint = InStartPoint;
-		Dir = InDir;
-		DeltaMag = InDeltaMag;
-		OutputFlags = InOutputFlags;
-		QueryFilterData = QueryFilter;
-
-		SQType = ESQType::Raycast;
-	}
-}
-
-void FSQCapture::StartCapturePhysXOverlap(const PxScene& Scene, const PxGeometry& InQueryGeom, const FTransform& WorldTM, const FQueryFilterData& QueryFilter, const FCollisionFilterData& FilterData, ICollisionQueryFilterCallbackBase& Callback)
-{
-	if (IsInGameThread())
-	{
-		bDiskDataIsChaos = false;
-		CapturePhysXFilterResults(Scene, FilterData, Callback);
-		//copy data
-		StartTM = WorldTM;
-		QueryFilterData = QueryFilter;
-
-		SQType = ESQType::Overlap;
-		SetPhysXGeometryData(InQueryGeom);
-	}
-}
-
-void FSQCapture::CapturePhysXFilterResults(const PxScene& TransientScene, const FCollisionFilterData& FilterData, ICollisionQueryFilterCallbackBase& Callback)
-{
-	const uint32 NumTransientActors = TransientScene.getNbActors(PxActorTypeFlag::eRIGID_STATIC | PxActorTypeFlag::eRIGID_DYNAMIC);
-	TArray<PxActor*> TransientActors; TransientActors.AddUninitialized(NumTransientActors);
-	if (NumTransientActors)
-	{
-		TransientScene.getActors(PxActorTypeFlag::eRIGID_STATIC | PxActorTypeFlag::eRIGID_DYNAMIC, TransientActors.GetData(), NumTransientActors);
-	}
-	
-	PxHitFlags QueryFlags;	//we know our callback throws this away so no need to store it or use real data
-	for (PxActor* TransientAct : TransientActors)
-	{
-		PxRigidActor* TransientActor = static_cast<PxRigidActor*>(TransientAct);
-
-		const uint32 NumTransientShapes = TransientActor->getNbShapes();
-		TArray<PxShape*> TransientShapes; TransientShapes.AddUninitialized(NumTransientShapes);
-		TransientActor->getShapes(TransientShapes.GetData(), NumTransientShapes);
-		TArray<TPair<PxShape*, ECollisionQueryHitType>> ShapeHitsArray; ShapeHitsArray.Reserve(NumTransientShapes);
-		for (PxShape* TransientShape : TransientShapes)
-		{
-			const PxQueryHitType::Enum Result = Callback.preFilter(U2PFilterData(FilterData), TransientShape, TransientActor, QueryFlags);
-			
-			//We must use the non transient shape/actor so that we can replay scene queries at runtime without serializing.
-			PxShape* NonTransientShape = static_cast<PxShape*>(PhysSerializer.FindObject((PxSerialObjectId)TransientShape));
-			ShapeHitsArray.Emplace(NonTransientShape, P2UCollisionQueryHitType(Result));
-			NonTransientToTransientShapes.Add(NonTransientShape, TransientShape);	//however, for serialization we must use the original shape/actor because conversion is done during load already
-		}
-
-		PxActor* NonTransientActor = static_cast<PxActor*>(PhysSerializer.FindObject((PxSerialObjectId)TransientActor));
-		PxActorToShapeHitsArray.Add(NonTransientActor, ShapeHitsArray);
-		NonTransientToTransientActors.Add(NonTransientActor, TransientActor);
-	}
-}
-
-
-template <typename THit>
-void EndCaptureHelper(PhysXInterface::FDynamicHitBuffer<THit>& Dest, const PxHitCallback<THit>& Results)
-{
-	Dest.block = Results.block;
-	Dest.hasBlock = Results.hasBlock;
-
-	if (Results.maxNbTouches == 0)
-	{
-		//we know this came from a single hit buffer because that's how UE uses the physx api.
-		//since we're putting it into a dynamic hit buffer we need to push the block into the hits array
-		if (Dest.hasBlock)
-		{
-			Dest.processTouches(&Dest.block, 1);
-		}
-	}
-	else
-	{
-		//we know this came from a dynamic hit buffer because that's how UE uses the physx api. Since it's dynamic block is already in the hits buffer
-		const PhysXInterface::FDynamicHitBuffer<THit>& DynamicResults = static_cast<const PhysXInterface::FDynamicHitBuffer<THit>&>(Results);
-		Dest.processTouches(DynamicResults.GetHits(), DynamicResults.GetNumHits());
-	}
-}
-
-void FSQCapture::EndCapturePhysXSweep(const PxHitCallback<PxSweepHit>& Results)
-{
-	if (IsInGameThread())
-	{
-		check(SQType == ESQType::Sweep);
-		EndCaptureHelper(PhysXSweepBuffer, Results);
-	}
-}
-
-void FSQCapture::EndCapturePhysXRaycast(const PxHitCallback<PxRaycastHit>& Results)
-{
-	if (IsInGameThread())
-	{
-		check(SQType == ESQType::Raycast);
-		EndCaptureHelper(PhysXRaycastBuffer, Results);
-	}
-}
-
-void FSQCapture::EndCapturePhysXOverlap(const PxHitCallback<PxOverlapHit>& Results)
-{
-	if (IsInGameThread())
-	{
-		check(SQType == ESQType::Overlap);
-		EndCaptureHelper(PhysXOverlapBuffer, Results);
-	}
-}
-
-template <typename THit>
-void FixupBufferPointers(FPhysTestSerializer& PhysSerializer, PhysXInterface::FDynamicHitBuffer<THit>& PhysXBuffer)
-{
-	auto FixupPointersLambda = [&PhysSerializer](PxActorShape& Hit)
-	{
-		Hit.actor = static_cast<PxRigidActor*>(PhysSerializer.FindObject((PxSerialObjectId)Hit.actor));
-		Hit.shape = static_cast<PxShape*>(PhysSerializer.FindObject((PxSerialObjectId)Hit.shape));
-	};
-
-	if (PhysXBuffer.hasBlock)
-	{
-		FixupPointersLambda(PhysXBuffer.block);
-	}
-	
-	const int32 NumHits = PhysXBuffer.GetNumHits();
-	for (int32 Idx = 0; Idx < NumHits; ++Idx)
-	{
-		THit& Hit = PhysXBuffer.GetHits()[Idx];
-		FixupPointersLambda(Hit);
-	}
-}
-
-ECollisionQueryHitType FSQCapture::GetFilterResult(const PxShape* Shape, const PxActor* Actor) const
-{
-	return GetFilterResultHelper(Shape, Actor, PxActorToShapeHitsArray);
-}
-
-void FSQCapture::CreatePhysXData()
-{
-	if (bDiskDataIsChaos || bPhysXDataReady)
-	{
-		return;
-	}
-
-	if (SQType != ESQType::Raycast)
-	{
-		AlignedDataHelper = MakeUnique<FPhysXSerializerData>(GeomData.Num());
-		FMemory::Memcpy(AlignedDataHelper->Data, GeomData.GetData(), GeomData.Num());
-
-		AlignedDataHelper->Registry = PxSerialization::createSerializationRegistry(*GPhysXSDK);
-		AlignedDataHelper->Collection = PxSerialization::createCollectionFromBinary(AlignedDataHelper->Data, *AlignedDataHelper->Registry);
-		if (PxBase* ColShape = AlignedDataHelper->Collection->find(ShapeCollectionID))
-		{
-			AlignedDataHelper->Shape = static_cast<PxShape*>(ColShape);
-			PhysXGeometry = AlignedDataHelper->Shape->getGeometry();
-		}
-		else
-		{
-			AlignedDataHelper.Reset();
-		}
-	}
-
-	FixupBufferPointers(PhysSerializer, PhysXRaycastBuffer);
-	FixupBufferPointers(PhysSerializer, PhysXSweepBuffer);
-	FixupBufferPointers(PhysSerializer, PhysXOverlapBuffer);
-
-	FilterCallback = MakeUnique<FSQCaptureFilterCallback>(*this);
-
-	bPhysXDataReady = true;
-}
-
-void FSQCapture::SetPhysXGeometryData(const PxGeometry& Geometry)
-{
-	check(AlignedDataHelper == nullptr);
-	check(SQType != ESQType::Raycast);
-
-	PxSerializationRegistry* Registry = PxSerialization::createSerializationRegistry(*GPhysXSDK);
-	PxCollection* Collection = PxCreateCollection();
-
-	//create a shape so we can serialize geometry
-	PxMaterial* Material = GPhysXSDK->createMaterial(1, 1, 1);
-	PxShape* Shape = GPhysXSDK->createShape(Geometry, *Material);
-	Collection->add(*Shape, ShapeCollectionID);
-
-	PxSerialization::complete(*Collection, *Registry);
-
-	GeomData.Empty();
-	FPhysXOutputStream Stream(&GeomData);
-	PxSerialization::serializeCollectionToBinary(Stream, *Collection, *Registry);
-	Material->release();
-	Shape->release();
-	Collection->release();
-	Registry->release();
-}
-
-FSQCapture::FPhysXSerializerData::FPhysXSerializerData(int32 NumBytes)
-	: Data(FMemory::Malloc(NumBytes, 128))
-	, Shape(nullptr)
-	, Collection(nullptr)
-	, Registry(nullptr)
-{
-}
-
-FSQCapture::FPhysXSerializerData::~FPhysXSerializerData()
-{
-	if (Collection)
-	{
-		//release all resources the collection created (calling release on the collection is not enough)
-		const uint32 NumObjects = Collection->getNbObjects();
-		TArray<PxBase*> Objects;
-		Objects.AddUninitialized(NumObjects);
-		Collection->getObjects(Objects.GetData(), NumObjects);
-		for (PxBase* Obj : Objects)
-		{
-			Obj->release();
-		}
-
-		Collection->release();
-		Registry->release();
-	}
-	FMemory::Free(Data);
-}
-
-#endif
-
-#if 0
-void PhysXQueryHitToChaosQueryHit(ChaosInterface::FQueryHit& ChaosHit, const PxQueryHit& PxHit, const FPhysTestSerializer& Serializer)
-{
-	ChaosHit.Actor = Serializer.PhysXActorToChaosHandle(PxHit.actor);
-	ChaosHit.Shape = Serializer.PhysXShapeToChaosImplicit(PxHit.shape);
-}
-
-void PhysXLocationHitToChaosLocationHit(ChaosInterface::FLocationHit& ChaosHit, const PxLocationHit& PxHit, const FPhysTestSerializer& Serializer)
-{
-	PhysXQueryHitToChaosQueryHit(ChaosHit, PxHit, Serializer);
-	ChaosHit.WorldPosition = P2UVector(PxHit.position);
-	ChaosHit.WorldNormal = P2UVector(PxHit.normal);
-	ChaosHit.Flags = P2UHitFlags(PxHit.flags);
-	ChaosHit.Distance = PxHit.distance;
-}
-
-void PhysXHitToChaosHit(ChaosInterface::FSweepHit& ChaosHit, const PxSweepHit& PxHit, const FPhysTestSerializer& Serializer)
-{
-	PhysXLocationHitToChaosLocationHit(ChaosHit, PxHit, Serializer);
-}
-
-void PhysXHitToChaosHit(ChaosInterface::FOverlapHit& ChaosHit, const PxOverlapHit& PxHit, const FPhysTestSerializer& Serializer)
-{
-	PhysXQueryHitToChaosQueryHit(ChaosHit, PxHit, Serializer);
-}
-
-void PhysXHitToChaosHit(ChaosInterface::FRaycastHit& ChaosHit, const PxRaycastHit& PxHit, const FPhysTestSerializer& Serializer)
-{
-	PhysXLocationHitToChaosLocationHit(ChaosHit, PxHit, Serializer);
-	ChaosHit.U = PxHit.u;
-	ChaosHit.V = PxHit.v;
-}
-
-template <typename TChaosHit, typename TPhysXHit>
-void PhysXToChaosBufferData(ChaosInterface::FSQHitBuffer<TChaosHit>& ChaosBuffer, const PhysXInterface::FDynamicHitBuffer<TPhysXHit>& PhysXBuffer, const FPhysTestSerializer& PhysSerializer)
-{
-	if (PhysXBuffer.hasBlock)
-	{
-		TChaosHit Hit;
-		PhysXHitToChaosHit(Hit, PhysXBuffer.block, PhysSerializer);
-		ChaosBuffer.SetBlockingHit(Hit);
-
-	}
-	const int32 NumHits = PhysXBuffer.GetNumHits();
-	for (int32 Idx = 0; Idx < NumHits; ++Idx)
-	{
-		TChaosHit Hit;
-		PhysXHitToChaosHit(Hit, PhysXBuffer.GetHits()[Idx], PhysSerializer);
-		ChaosBuffer.AddTouchingHit(Hit);
-	}
-}
-
-void FSQCapture::CreateChaosFilterResults()
-{
-#if WITH_PHYSX
-	for (auto Itr : PxActorToShapeHitsArray)
-	{
-		Chaos::FGeometryParticle* Actor = PhysSerializer.PhysXActorToChaosHandle(Itr.Key);
-		const auto& ShapesArray = Actor->ShapesArray();
-		TArray<TPair<Chaos::FPerShapeData*, ECollisionQueryHitType>> FilterResults;
-
-		const auto& Pairs = Itr.Value;
-		int32 Idx = 0;
-		for (const auto& Pair : Pairs)
-		{
-			FilterResults.Add(TPair<Chaos::FPerShapeData*, ECollisionQueryHitType>(ShapesArray[Idx++].Get(), Pair.Value));
-		}
-		ChaosActorToShapeHitsArray.Add(Actor, FilterResults);
-	}
-#endif
-}
-
-void FSQCapture::CreateChaosDataFromPhysX()
-{
-	using namespace Chaos;
-	
-	if (bDiskDataIsChaos || bChaosDataReady)
-	{
-		return;
-	}
-
-	if (AlignedDataHelper && AlignedDataHelper->Shape)
-	{
-		TUniquePtr<TImplicitObjectTransformed<float, 3>> TransformedObj = PxShapeToChaosGeom(AlignedDataHelper->Shape);
-		//we know the dummy px shape has no transform so we want to use the inner object
-		ChaosGeometry = TransformedObj->GetTransformedObject();
-		ChaosOwnerObject = MoveTemp(TransformedObj);
-	}
-
-#if WITH_CHAOS
-	PhysXToChaosBufferData(ChaosRaycastBuffer, PhysXRaycastBuffer, PhysSerializer);
-	PhysXToChaosBufferData(ChaosSweepBuffer, PhysXSweepBuffer, PhysSerializer);
-	PhysXToChaosBufferData(ChaosOverlapBuffer, PhysXOverlapBuffer, PhysSerializer);
-#endif // WITH_CHAOS
-
-	CreateChaosFilterResults();
-	bChaosDataReady = true;
-}
-#endif
diff --git a/Engine/Source/Runtime/PhysicsCore/Private/SceneQueryChaosImp.h b/Engine/Source/Runtime/PhysicsCore/Private/SceneQueryChaosImp.h
--- a/Engine/Source/Runtime/PhysicsCore/Private/SceneQueryChaosImp.h
+++ b/Engine/Source/Runtime/PhysicsCore/Private/SceneQueryChaosImp.h
@@ -2,8 +2,6 @@
 
 #pragma once
 
-#if WITH_CHAOS
-
 inline bool LowLevelRaycastImp(const FVector& Start, const FVector& Dir, float DeltaMag, const Chaos::FImplicitObject& Shape, const FTransform ActorTM, EHitFlags OutputFlags, FHitRaycast& Hit)
 {
 	//TODO_SQ_IMPLEMENTATION
@@ -20,5 +18,4 @@ inline bool LowLevelOverlapImp(const FTransform& GeomPose, const FPhysicsGeometr
 {
 	//TODO_SQ_IMPLEMENTATION
 	return false;
-}
-#endif // WITH_CHAOS
\ No newline at end of file
+}
\ No newline at end of file
diff --git a/Engine/Source/Runtime/PhysicsCore/Private/SceneQueryPhysXImp.h b/Engine/Source/Runtime/PhysicsCore/Private/SceneQueryPhysXImp.h
--- a/Engine/Source/Runtime/PhysicsCore/Private/SceneQueryPhysXImp.h
+++ b/Engine/Source/Runtime/PhysicsCore/Private/SceneQueryPhysXImp.h
@@ -1,36 +1,3 @@
 // Copyright Epic Games, Inc. All Rights Reserved.
 
 #pragma once
-
-#if WITH_PHYSX && PHYSICS_INTERFACE_PHYSX
-
-#include "PhysXPublicCore.h"
-
-FORCEINLINE bool LowLevelRaycastImp(const PxVec3& Start, const PxVec3& Dir, float DeltaMag, const PxShape& Shape, const PxTransform ActorTM, PxHitFlags OutputFlags, PxRaycastHit& Hit)
-{
-	const PxGeometryHolder GeomHolder = Shape.getGeometry();
-	const PxGeometry& Geom = GeomHolder.any();
-	const PxTransform GeomTM = ActorTM * Shape.getLocalPose();
-
-	return !!PxGeometryQuery::raycast(Start, Dir, Geom, GeomTM, DeltaMag, OutputFlags, 1, &Hit);
-							
-}
-
-FORCEINLINE bool LowLevelSweepImp(const PxTransform& StartTM, const PxVec3& Dir, float DeltaMag, const PxGeometry& SweepGeom, const PxShape& Shape, const PxTransform ActorTM, PxHitFlags OutputFlags, PxSweepHit& Hit)
-{
-	const PxGeometryHolder GeomHolder = Shape.getGeometry();
-	const PxGeometry& ShapeGeom = GeomHolder.any();
-	const PxTransform ShapeGeomTM = ActorTM * Shape.getLocalPose();
-
-	return PxGeometryQuery::sweep(Dir, DeltaMag, SweepGeom, StartTM, ShapeGeom, ShapeGeomTM, Hit, OutputFlags);
-}
-
-FORCEINLINE bool LowLevelOverlapImp(const PxTransform& GeomPose, const PxGeometry& OverlapGeom, const PxShape& Shape, const PxTransform ActorTM, PxOverlapHit& Overlap)
-{
-	const PxGeometryHolder GeomHolder = Shape.getGeometry();
-	const PxGeometry& ShapeGeom = GeomHolder.any();
-	const PxTransform ShapeGeomTM = ActorTM * Shape.getLocalPose();
-
-	return PxGeometryQuery::overlap(OverlapGeom, GeomPose, ShapeGeom, ShapeGeomTM);
-}
-#endif // WITH_PHYSX 
\ No newline at end of file
diff --git a/Engine/Source/Runtime/PhysicsCore/Public/BodySetupCore.h b/Engine/Source/Runtime/PhysicsCore/Public/BodySetupCore.h
--- a/Engine/Source/Runtime/PhysicsCore/Public/BodySetupCore.h
+++ b/Engine/Source/Runtime/PhysicsCore/Public/BodySetupCore.h
@@ -39,10 +39,4 @@ public:
 	/** Collision Type for this body. This eventually changes response to collision to others **/
 	UPROPERTY(EditAnywhere,Category=Collision)
 	TEnumAsByte<enum EBodyCollisionResponse::Type> CollisionReponse;
-
-	
-#if PHYSICS_INTERFACE_PHYSX
-	/** Physics triangle mesh, created from cooked data in CreatePhysicsMeshes */
-	TArray<physx::PxTriangleMesh*> TriMeshes;
-#endif
 };
\ No newline at end of file
diff --git a/Engine/Source/Runtime/PhysicsCore/Public/Chaos/ChaosEngineInterface.h b/Engine/Source/Runtime/PhysicsCore/Public/Chaos/ChaosEngineInterface.h
--- a/Engine/Source/Runtime/PhysicsCore/Public/Chaos/ChaosEngineInterface.h
+++ b/Engine/Source/Runtime/PhysicsCore/Public/Chaos/ChaosEngineInterface.h
@@ -197,8 +197,6 @@ enum ELinearConstraintMotion
     If(QueryIgnoreMask & ShapeFilter != 0) filter out */
 typedef uint8 FMaskFilter;
 
-#if WITH_CHAOS
-
 namespace Chaos
 {
 	class FBVHParticles;
@@ -432,9 +430,7 @@ public:
 	static void SetMaxDepenetrationVelocity_AssumesLocked(const FPhysicsActorHandle& InActorReference,float InMaxDepenetrationVelocity);
 
 	static FVector GetWorldVelocityAtPoint_AssumesLocked(const FPhysicsActorHandle& InActorReference,const FVector& InPoint);
-#if WITH_CHAOS
 	static FVector GetWorldVelocityAtPoint_AssumesLocked(const Chaos::FRigidBodyHandle_Internal* InActorReference, const FVector& InPoint);
-#endif
 
 	static FTransform GetComTransform_AssumesLocked(const FPhysicsActorHandle& InActorReference);
 	static FTransform GetComTransformLocal_AssumesLocked(const FPhysicsActorHandle& InActorReference);
@@ -544,36 +540,3 @@ public:
 	static void SetIsQueryShape(const FPhysicsShapeHandle& InShape,bool bIsQueryShape);
 	static void SetLocalTransform(const FPhysicsShapeHandle& InShape,const FTransform& NewLocalTransform);
 };
-
-#elif WITH_ENGINE //temp physx code to make moving code out of Engine easier
-
-/**
- * Wrapper for internal PhysX materials
- */
-
-
-namespace physx
-{
-class PxMaterial;
-}
-
-struct PHYSICSCORE_API FPhysicsMaterialHandle_PhysX
-{
-	FPhysicsMaterialHandle_PhysX() : Material(nullptr) {}
-	explicit FPhysicsMaterialHandle_PhysX(physx::PxMaterial* InMaterial) : Material(InMaterial) {}
-
-	bool IsValid() const { return Material != nullptr; }
-
-	physx::PxMaterial* Material;
-};
-
-class FChaosEngineInterface
-{
-public:
-	static FPhysicsMaterialHandle CreateMaterial(const UPhysicalMaterial* InMaterial);
-	static void UpdateMaterial(FPhysicsMaterialHandle_PhysX& InHandle,UPhysicalMaterial* InMaterial);
-	static void ReleaseMaterial(FPhysicsMaterialHandle_PhysX& InHandle);
-	static void SetUserData(FPhysicsMaterialHandle_PhysX& InHandle,void* InUserData);
-};
-
-#endif
\ No newline at end of file
diff --git a/Engine/Source/Runtime/PhysicsCore/Public/Chaos/ChaosScene.h b/Engine/Source/Runtime/PhysicsCore/Public/Chaos/ChaosScene.h
--- a/Engine/Source/Runtime/PhysicsCore/Public/Chaos/ChaosScene.h
+++ b/Engine/Source/Runtime/PhysicsCore/Public/Chaos/ChaosScene.h
@@ -2,8 +2,6 @@
 
 #pragma once
 
-#if WITH_CHAOS
-
 #include "CoreMinimal.h"
 #include "UObject/UObjectGlobals.h"
 #include "UObject/GCObject.h"
@@ -114,9 +112,7 @@ public:
 
 	void AddActorsToScene_AssumesLocked(TArray<FPhysicsActorHandle>& InHandles,const bool bImmediate=true);
 	void RemoveActorFromAccelerationStructure(FPhysicsActorHandle Actor);
-#if WITH_CHAOS
 	void RemoveActorFromAccelerationStructureImp(Chaos::FGeometryParticle* Particle);
-#endif
 	void UpdateActorsInAccelerationStructure(const TArrayView<FPhysicsActorHandle>& Actors);
 	void UpdateActorInAccelerationStructure(const FPhysicsActorHandle& Actor);
 
@@ -194,5 +190,3 @@ private:
 
 	FChaosSceneSimCallback* SimCallback;
 };
-
-#endif
\ No newline at end of file
diff --git a/Engine/Source/Runtime/PhysicsCore/Public/ChaosInterfaceWrapperCore.h b/Engine/Source/Runtime/PhysicsCore/Public/ChaosInterfaceWrapperCore.h
--- a/Engine/Source/Runtime/PhysicsCore/Public/ChaosInterfaceWrapperCore.h
+++ b/Engine/Source/Runtime/PhysicsCore/Public/ChaosInterfaceWrapperCore.h
@@ -6,9 +6,7 @@
 #include "PhysicsInterfaceWrapperShared.h"
 #include "PhysicsInterfaceTypesCore.h"
 
-#if WITH_PHYSX
 #include "PhysXPublicCore.h"
-#endif
 
 class UPhysicalMaterial;
 
@@ -27,11 +25,7 @@ struct FDummyPhysActor {};
 template<typename DummyT>
 struct FDummyCallback {};
 
-#if PHYSICS_INTERFACE_PHYSX
-using FQueryFilterData = PxQueryFilterData;
-#elif WITH_CHAOS
 using FQueryFilterData = FChaosQueryFilterData;
-#endif
 
 /** We use this struct so that if no conversion is needed in another API, we can avoid the copy (if we think that's critical) */
 struct FPhysicsRaycastInputAdapater
@@ -74,7 +68,6 @@ struct FPhysicsOverlapInputAdapater
 	FTransform GeomPose;
 };
 
-#if WITH_CHAOS
 /** This is used to add debug data to scene query visitors in non-shipping builds */
 struct FQueryDebugParams
 {
@@ -93,7 +86,6 @@ struct FQueryDebugParams
 	constexpr bool IsExternalQuery() const { return true; }
 #endif
 };
-#endif
 
 extern PHYSICSCORE_API FCollisionFilterData GetQueryFilterData(const Chaos::FPerShapeData& Shape);
 extern PHYSICSCORE_API FCollisionFilterData GetSimulationFilterData(const Chaos::FPerShapeData& Shape);
@@ -262,6 +254,4 @@ bool GetHasBlock(const FSQHitBuffer<HitType>& Callback)
 
 } // namespace ChaosInterface
 
-#if WITH_CHAOS && (!defined(PHYSICS_INTERFACE_PHYSX) || !PHYSICS_INTERFACE_PHYSX)
 using namespace ChaosInterface;
-#endif
diff --git a/Engine/Source/Runtime/PhysicsCore/Public/CollisionQueryFilterCallbackCore.h b/Engine/Source/Runtime/PhysicsCore/Public/CollisionQueryFilterCallbackCore.h
--- a/Engine/Source/Runtime/PhysicsCore/Public/CollisionQueryFilterCallbackCore.h
+++ b/Engine/Source/Runtime/PhysicsCore/Public/CollisionQueryFilterCallbackCore.h
@@ -1,18 +1,9 @@
 // Copyright Epic Games, Inc. All Rights Reserved.
 
-
 #pragma once
 
 #include "PhysicsInterfaceTypesCore.h"
-
-#if PHYSICS_INTERFACE_PHYSX
-#include "PhysXPublicCore.h"
-#include "PhysXInterfaceWrapperCore.h"
-#endif
-#if WITH_CHAOS
 #include "ChaosInterfaceWrapperCore.h"
-#endif
-
 
 /**
  *
@@ -31,9 +22,6 @@ namespace Chaos
 }
 
 class ICollisionQueryFilterCallbackBase
-#if PHYSICS_INTERFACE_PHYSX
-	: public PxQueryFilterCallback	//this is a bit of a hack until physx support is gone. Means any new non physx callback will have dummy pre/postFilter functions for physx api
-#endif
 {
 public:
 	
@@ -43,11 +31,6 @@ public:
 
 	virtual ECollisionQueryHitType PostFilter(const FCollisionFilterData& FilterData, const ChaosInterface::FPTQueryHit& Hit) = 0;
 	virtual ECollisionQueryHitType PreFilter(const FCollisionFilterData& FilterData, const Chaos::FPerShapeData& Shape, const Chaos::FGeometryParticleHandle& Actor) = 0;
-
-#if  PHYSICS_INTERFACE_PHYSX
-	virtual ECollisionQueryHitType PostFilter(const FCollisionFilterData& FilterData, const physx::PxQueryHit& Hit) = 0;
-	virtual ECollisionQueryHitType PreFilter(const FCollisionFilterData& FilterData, const physx::PxShape& Shape, physx::PxRigidActor& Actor) = 0;
-#endif
 };
 
 class PHYSICSCORE_API FBlockAllQueryCallback : public ICollisionQueryFilterCallbackBase
@@ -59,13 +42,6 @@ public:
 
 	virtual ECollisionQueryHitType PostFilter(const FCollisionFilterData& FilterData, const ChaosInterface::FPTQueryHit& Hit) override { return ECollisionQueryHitType::Block; }
 	virtual ECollisionQueryHitType PreFilter(const FCollisionFilterData& FilterData, const Chaos::FPerShapeData& Shape, const Chaos::FGeometryParticleHandle& Actor) override { return ECollisionQueryHitType::Block; }
-
-#if PHYSICS_INTERFACE_PHYSX
-	virtual ECollisionQueryHitType PostFilter(const FCollisionFilterData& FilterData, const physx::PxQueryHit& Hit) override { return ECollisionQueryHitType::Block; }
-	virtual ECollisionQueryHitType PreFilter(const FCollisionFilterData& FilterData, const physx::PxShape& Shape, physx::PxRigidActor& Actor) override { return ECollisionQueryHitType::Block; }
-	virtual PxQueryHitType::Enum preFilter(const PxFilterData& filterData, const PxShape* shape, const PxRigidActor* actor, PxHitFlags& queryFlags) override { return PxQueryHitType::eBLOCK; }
-	virtual PxQueryHitType::Enum postFilter(const PxFilterData& filterData, const PxQueryHit& hit) override { return PxQueryHitType::eBLOCK; }
-#endif
 };
 
 class PHYSICSCORE_API FOverlapAllQueryCallback : public ICollisionQueryFilterCallbackBase
@@ -77,23 +53,4 @@ public:
 
 	virtual ECollisionQueryHitType PostFilter(const FCollisionFilterData& FilterData, const ChaosInterface::FPTQueryHit& Hit) override { return ECollisionQueryHitType::Touch; }
 	virtual ECollisionQueryHitType PreFilter(const FCollisionFilterData& FilterData, const Chaos::FPerShapeData& Shape, const Chaos::FGeometryParticleHandle& Actor) override { return ECollisionQueryHitType::Touch; }
-
-#if PHYSICS_INTERFACE_PHYSX
-	virtual ECollisionQueryHitType PostFilter(const FCollisionFilterData& FilterData, const physx::PxQueryHit& Hit) override { return ECollisionQueryHitType::Touch; }
-	virtual ECollisionQueryHitType PreFilter(const FCollisionFilterData& FilterData, const physx::PxShape& Shape, physx::PxRigidActor& Actor) override { return ECollisionQueryHitType::Touch; }
-	virtual PxQueryHitType::Enum preFilter(const PxFilterData& filterData, const PxShape* shape, const PxRigidActor* actor, PxHitFlags& queryFlags) override { return PxQueryHitType::eTOUCH; }
-	virtual PxQueryHitType::Enum postFilter(const PxFilterData& filterData, const PxQueryHit& hit) override { return PxQueryHitType::eTOUCH; }
-#endif
-};
-
-#if PHYSICS_INTERFACE_PHYSX
-inline PxQueryHitType::Enum U2PCollisionQueryHitType(ECollisionQueryHitType HitType)
-{
-	return (PxQueryHitType::Enum)HitType;
-}
-
-inline ECollisionQueryHitType P2UCollisionQueryHitType(PxQueryHitType::Enum HitType)
-{
-	return (ECollisionQueryHitType)HitType;
-}
-#endif
\ No newline at end of file
+};
\ No newline at end of file
diff --git a/Engine/Source/Runtime/PhysicsCore/Public/PhysTestSerializer.h b/Engine/Source/Runtime/PhysicsCore/Public/PhysTestSerializer.h
--- a/Engine/Source/Runtime/PhysicsCore/Public/PhysTestSerializer.h
+++ b/Engine/Source/Runtime/PhysicsCore/Public/PhysTestSerializer.h
@@ -17,22 +17,9 @@
 #endif
 
 // Utility used for serializing just physics data. This is meant for only the physics engine data (physx or chaos, not any unreal side). It is not meant to be used for actual serialization
-// This utility allows for moving back and forth between physx and chaos data. This is not a proper upgrade path
 
 #if PHYS_TEST_SERIALIZER
 
-#if PHYSICS_INTERFACE_PHYSX
-namespace physx
-{
-	class PxScene;
-	class PxSerializationRegistry;
-	class PxCollection;
-	class PxBase;
-	class PxActor;
-	class PxShape;
-}
-#endif
-
 namespace Chaos
 {
 
@@ -54,10 +41,6 @@ public:
 	//Set the data from an external source. This will obliterate any existing data. Make sure you are not holding on to old internal data as it will go away
 	void SetPhysicsData(Chaos::FPBDRigidsEvolution& ChaosEvolution);
 
-#if PHYSICS_INTERFACE_PHYSX
-	void SetPhysicsData(physx::PxScene& Scene);
-#endif
-
 	const Chaos::FChaosArchiveContext* GetChaosContext() const
 	{
 		return ChaosContext.Get();
@@ -75,11 +58,6 @@ public:
 		if (SQCapture)
 		{
 			//todo: this sucks, find a better way to create data instead of doing it lazily
-#if PHYSICS_INTERFACE_PHYSX
-			GetPhysXData();
-			SQCapture->CreatePhysXData();
-#endif
-
 			GetChaosData();
 #if 0 
 			SQCapture->CreateChaosDataFromPhysX();
@@ -101,37 +79,11 @@ public:
 		return ChaosEvolution.Get();
 	}
 
-#if PHYSICS_INTERFACE_PHYSX
-	physx::PxScene* GetPhysXData()
-	{
-		if (!bDiskDataIsChaos)	//don't support chaos to physx
-		{
-			if (!AlignedDataHelper)
-			{
-				CreatePhysXData();
-			}
-
-			return AlignedDataHelper->PhysXScene;
-		}
-		return nullptr;
-	}
-
-	physx::PxBase* FindObject(uint64 Id);
-
-#if WITH_CHAOS
-	Chaos::FGeometryParticle* PhysXActorToChaosHandle(physx::PxActor* Actor) const { return PxActorToChaosHandle.FindChecked(Actor)->GTGeometryParticle(); }
-	Chaos::FPerShapeData* PhysXShapeToChaosImplicit(physx::PxShape* Shape) const { return PxShapeToChaosShapes.FindRef(Shape); }
-#endif
-#endif
-
 private:
 
 #if 0
 	void CreateChaosData();
 #endif 
-#if PHYSICS_INTERFACE_PHYSX
-	void CreatePhysXData();
-#endif
 
 private:
 
@@ -150,27 +102,6 @@ private:
 	TUniquePtr<Chaos::FChaosArchiveContext> ChaosContext;
 
 	FCustomVersionContainer ArchiveVersion;
-
-#if PHYSICS_INTERFACE_PHYSX
-	struct FPhysXSerializerData
-	{
-		FPhysXSerializerData(int32 NumBytes)
-			: Data(FMemory::Malloc(NumBytes, 128))
-			, PhysXScene(nullptr)
-			, Collection(nullptr)
-			, Registry(nullptr)
-		{}
-		~FPhysXSerializerData();
-		void* Data;
-		physx::PxScene* PhysXScene;
-		physx::PxCollection* Collection;
-		physx::PxSerializationRegistry* Registry;
-	};
-	TUniquePtr<FPhysXSerializerData> AlignedDataHelper;
-
-	TMap<physx::PxActor*, Chaos::TGeometryParticleHandle<float, 3>*> PxActorToChaosHandle;
-	TMap<physx::PxShape*, Chaos::FPerShapeData*> PxShapeToChaosShapes;
-#endif
 };
 
 #endif
diff --git a/Engine/Source/Runtime/PhysicsCore/Public/PhysXIncludes.h b/Engine/Source/Runtime/PhysicsCore/Public/PhysXIncludes.h
--- a/Engine/Source/Runtime/PhysicsCore/Public/PhysXIncludes.h
+++ b/Engine/Source/Runtime/PhysicsCore/Public/PhysXIncludes.h
@@ -4,116 +4,3 @@
 #pragma once
 
 #include "CoreMinimal.h"
-
-#if PHYSICS_INTERFACE_PHYSX
-
-#ifdef PHYSX_API
-#undef PHYSX_API	//Our build system treats PhysX as a module and automatically defines this. PhysX has its own way of defining this.
-#endif
-
-MSVC_PRAGMA(warning(push))
-MSVC_PRAGMA(warning(disable : 4946)) // reinterpret_cast used between related classes
-
-PRAGMA_DISABLE_SHADOW_VARIABLE_WARNINGS
-
-#if USING_CODE_ANALYSIS
-	MSVC_PRAGMA(warning(push))
-	MSVC_PRAGMA(warning(disable : ALL_CODE_ANALYSIS_WARNINGS))
-#endif	// USING_CODE_ANALYSIS
-
-#if PLATFORM_MICROSOFT
-	#if PLATFORM_64BITS
-		#pragma pack(push,16)
-	#elif PLATFORM_32BITS
-		#pragma pack(push,8)
-	#endif
-#elif PLATFORM_UNIX
-	#if (PLATFORM_CPU_X86_FAMILY && !PLATFORM_64BITS)
-		#pragma pack(push,16)
-	#endif
-	#ifdef __clang__
-		#pragma clang diagnostic push
-		#pragma clang diagnostic ignored "-Wreturn-type-c-linkage"	// ignore PxPhysics.h PxGetPhysics() returning C++ type - see UE-49412
-	#endif
-#endif
-
-THIRD_PARTY_INCLUDES_START
-#include "Px.h"
-#include "PxPhysicsAPI.h"
-#include "PxRenderBuffer.h"
-#include "PxExtensionsAPI.h"
-#include "PxCollectionExt.h"
-#include "PxPvd.h"
-
-#include "PxImmediateMode.h"
-
-// utils
-#include "PxGeometryQuery.h"
-#include "PxMeshQuery.h"
-#include "PxTriangle.h"
-THIRD_PARTY_INCLUDES_END
-
-// APEX
-#if WITH_APEX
-
-#ifdef APEX_API
-#undef APEX_API	//Our build system treats PhysX as a module and automatically defines this. PhysX has its own way of defining this.
-#endif
-
-// Framework
-THIRD_PARTY_INCLUDES_START
-#include "Apex.h"
-THIRD_PARTY_INCLUDES_END
-
-// Modules
-//These are still included until we get all of APEX into its own plugin
-#include "destructible/ModuleDestructible.h"
-#include "destructible/DestructibleAsset.h"
-#include "destructible/DestructibleActor.h"
-
-THIRD_PARTY_INCLUDES_START
-
-#if WITH_APEX_CLOTHING
-#include "clothing/ModuleClothing.h"
-#include "clothing/ClothingAsset.h"
-#include "clothing/ClothingActor.h"
-#include "clothing/ClothingCollision.h"
-#endif
-
-#if WITH_APEX_LEGACY
-#include "ModuleLegacy.h"
-#endif
-
-// Utilities
-#include "NvParamUtils.h"
-
-THIRD_PARTY_INCLUDES_END
-
-#endif // #if WITH_APEX
-
-#if PLATFORM_MICROSOFT
-	#if PLATFORM_64BITS
-		#pragma pack(pop)
-	#elif PLATFORM_32BITS
-		#pragma pack(pop)
-	#endif
-#elif PLATFORM_UNIX
-	#if (PLATFORM_CPU_X86_FAMILY && !PLATFORM_64BITS)
-		#pragma pack(pop)
-	#endif
-	#ifdef __clang__
-		#pragma clang diagnostic pop
-	#endif
-#endif
-
-#if USING_CODE_ANALYSIS
-	MSVC_PRAGMA(warning(pop))
-#endif	// USING_CODE_ANALYSIS
-
-MSVC_PRAGMA(warning(pop))
-
-PRAGMA_ENABLE_SHADOW_VARIABLE_WARNINGS
-
-using namespace physx;
-
-#endif // WITH_PHYSX  
\ No newline at end of file
diff --git a/Engine/Source/Runtime/PhysicsCore/Public/PhysXInterfaceWrapperCore.h b/Engine/Source/Runtime/PhysicsCore/Public/PhysXInterfaceWrapperCore.h
--- a/Engine/Source/Runtime/PhysicsCore/Public/PhysXInterfaceWrapperCore.h
+++ b/Engine/Source/Runtime/PhysicsCore/Public/PhysXInterfaceWrapperCore.h
@@ -1,404 +1,2 @@
 // Copyright Epic Games, Inc. All Rights Reserved.
 #pragma once
-
-#if PHYSICS_INTERFACE_PHYSX
-#include "PhysXPublicCore.h"
-#include "PhysicsInterfaceUtilsCore.h"
-#include "PhysicsInterfaceWrapperShared.h"
-#include "ChaosSQTypes.h"
-
-inline EQueryFlags P2UQueryFlags(PxQueryFlags Flags)
-{
-	FQueryFlags Result = EQueryFlags::None;
-	if (Flags & PxQueryFlag::ePREFILTER)
-	{
-		Result |= EQueryFlags::PreFilter;
-	}
-
-	if (Flags & PxQueryFlag::ePOSTFILTER)
-	{
-		Result |= EQueryFlags::PostFilter;
-	}
-
-	if (Flags & PxQueryFlag::eANY_HIT)
-	{
-		Result |= EQueryFlags::AnyHit;
-	}
-
-	return Result.QueryFlags;
-}
-
-inline PxQueryFlags U2PQueryFlags(FQueryFlags Flags)
-{
-	uint32 Result = 0;
-	if (Flags & EQueryFlags::PreFilter)
-	{
-		Result |= PxQueryFlag::ePREFILTER;
-	}
-
-	if (Flags & EQueryFlags::PostFilter)
-	{
-		Result |= PxQueryFlag::ePOSTFILTER;
-	}
-
-	if (Flags & EQueryFlags::AnyHit)
-	{
-		Result |= PxQueryFlag::eANY_HIT;
-	}
-
-	return (PxQueryFlags)Result;
-}
-
-FORCEINLINE PxHitFlags U2PHitFlags(const FHitFlags& Flags)
-{
-	uint32 Result = 0;
-	if (Flags & EHitFlags::Position)
-	{
-		Result |= PxHitFlag::ePOSITION;
-	}
-
-	if (Flags & EHitFlags::Normal)
-	{
-		Result |= PxHitFlag::eNORMAL;
-	}
-
-	if (Flags & EHitFlags::Distance)
-	{
-		Result |= PxHitFlag::eDISTANCE;
-	}
-
-	if (Flags & EHitFlags::UV)
-	{
-		Result |= PxHitFlag::eUV;
-	}
-
-	if (Flags & EHitFlags::MTD)
-	{
-		Result |= PxHitFlag::eMTD;
-	}
-
-	if (Flags & EHitFlags::FaceIndex)
-	{
-		Result |= PxHitFlag::eFACE_INDEX;
-	}
-
-	return (PxHitFlags)Result;
-}
-
-FORCEINLINE EHitFlags P2UHitFlags(const PxHitFlags& Flags)
-{
-	FHitFlags Result = EHitFlags::None;
-	if (Flags & PxHitFlag::ePOSITION)
-	{
-		Result |= EHitFlags::Position;
-	}
-
-	if (Flags & PxHitFlag::eDISTANCE)
-	{
-		Result |= EHitFlags::Distance;
-	}
-
-	if (Flags & PxHitFlag::eNORMAL)
-	{
-		Result |= EHitFlags::Normal;
-	}
-
-	if (Flags & PxHitFlag::eUV)
-	{
-		Result |= EHitFlags::UV;
-	}
-
-	if (Flags & PxHitFlag::eMTD)
-	{
-		Result |= EHitFlags::MTD;
-	}
-
-	if (Flags & PxHitFlag::eFACE_INDEX)
-	{
-		Result |= EHitFlags::FaceIndex;
-	}
-
-	return Result.HitFlags;
-}
-
-namespace PhysXInterface
-{
-
-FORCEINLINE bool HadInitialOverlap(const PxLocationHit& Hit)
-{
-	return Hit.hadInitialOverlap();
-}
-
-FORCEINLINE PxShape* GetShape(const PxLocationHit& Hit)
-{
-	return Hit.shape;
-}
-
-FORCEINLINE PxShape* GetShape(const PxOverlapHit& Hit)
-{
-	return Hit.shape;
-}
-
-FORCEINLINE PxRigidActor* GetActor(const PxLocationHit& Hit)
-{
-	return Hit.actor;
-}
-
-FORCEINLINE PxRigidActor* GetActor(const PxOverlapHit& Hit)
-{
-	return Hit.actor;
-}
-
-FORCEINLINE float GetDistance(const PxLocationHit& Hit)
-{
-	return Hit.distance;
-}
-
-template <typename HitType>
-HitType* GetBlock(PxHitCallback<HitType>& Callback)
-{
-	return &Callback.block;
-}
-
-template <typename HitType>
-bool GetHasBlock(const PxHitCallback<HitType>& Callback)
-{
-	return Callback.hasBlock;
-}
-
-FORCEINLINE FVector GetPosition(const PxLocationHit& Hit)
-{
-	return P2UVector(Hit.position);
-}
-
-FORCEINLINE FVector GetNormal(const PxLocationHit& Hit)
-{
-	return P2UVector(Hit.normal);
-}
-
-FORCEINLINE FHitFlags GetFlags(const PxLocationHit& Hit)
-{
-	return P2UHitFlags(Hit.flags);
-}
-
-FORCEINLINE void SetFlags(PxLocationHit& Hit, FHitFlags Flags)
-{
-	Hit.flags = U2PHitFlags(Flags);
-}
-
-FORCEINLINE uint32 GetInternalFaceIndex(const PxLocationHit& Hit)
-{
-	return Hit.faceIndex;
-}
-
-FORCEINLINE void SetInternalFaceIndex(PxLocationHit& Hit, uint32 FaceIndex)
-{
-	Hit.faceIndex = FaceIndex;
-}
-
-FORCEINLINE FCollisionFilterData GetQueryFilterData(const PxShape& Shape)
-{
-	return P2UFilterData(Shape.getQueryFilterData());
-}
-
-FORCEINLINE FCollisionFilterData GetSimulationFilterData(const PxShape& Shape)
-{
-	return P2UFilterData(Shape.getSimulationFilterData());
-}
-
-FORCEINLINE ECollisionShapeType GetType(const PxGeometry& Geom)
-{
-	return P2UGeometryType(Geom.getType());
-}
-
-FORCEINLINE ECollisionShapeType GetGeometryType(const PxShape& Shape)
-{
-	return P2UGeometryType(Shape.getGeometryType());
-}
-
-FORCEINLINE PxMaterial* GetMaterialFromInternalFaceIndex(const PxShape& Shape, const FPhysicsActor& Actor, uint32 InternalFaceIndex)
-{
-	return Shape.getMaterialFromInternalFaceIndex(InternalFaceIndex);
-}
-
-FORCEINLINE uint32 GetInvalidPhysicsFaceIndex()
-{
-	// Sentinel for invalid query results.
-	static const PxQueryHit InvalidQueryHit;
-	return InvalidQueryHit.faceIndex;
-}
-
-FORCEINLINE_DEBUGGABLE bool IsInvalidFaceIndex(PxU32 faceIndex)
-{
-	checkfSlow(GetInvalidPhysicsFaceIndex() == 0xFFFFffff, TEXT("Engine code needs fixing: PhysX invalid face index sentinel has changed or is not part of default PxQueryHit!"));
-	return (faceIndex == 0xFFFFffff);
-}
-
-using FQueryFilterData = PxQueryFilterData;
-
-template <typename HitType>
-class FSingleHitBuffer : public PxHitBuffer<HitType>
-{
-public:
-	FSingleHitBuffer()
-		: PxHitBuffer<HitType>()
-	{
-
-	}
-
-};
-
-using FPhysicsSweepBuffer = FSingleHitBuffer<PxSweepHit>;
-using FPhysicsRaycastBuffer = FSingleHitBuffer<PxRaycastHit>;
-
-#define HIT_BUFFER_SIZE							512		// Hit buffer size for traces and sweeps. This is the total size allowed for sync + async tests.
-static_assert(HIT_BUFFER_SIZE > 0, "Invalid PhysX hit buffer size.");
-
-//todo(ocohen): remove all of this - it's not clear that we care about dynamic dispatch at this layer
-template<typename HitType>
-class FDynamicHitBuffer : public PxHitCallback<HitType>
-{
-private:
-	/** Hit buffer used to provide hits via processTouches */
-	TTypeCompatibleBytes<HitType> HitBuffer[HIT_BUFFER_SIZE];
-
-	/** Hits encountered. Can be larger than HIT_BUFFER_SIZE */
-	TArray<TTypeCompatibleBytes<HitType>, TInlineAllocator<HIT_BUFFER_SIZE>> Hits;
-
-public:
-	FDynamicHitBuffer()
-		: PxHitCallback<HitType>((HitType*)HitBuffer, HIT_BUFFER_SIZE)
-	{}
-
-	virtual PxAgain processTouches(const HitType* buffer, PxU32 nbHits) override
-	{
-		Hits.Append((TTypeCompatibleBytes<HitType>*)buffer, nbHits);
-		return true;
-	}
-
-	virtual void finalizeQuery() override
-	{
-		if (this->hasBlock)
-		{
-			// copy blocking hit to hits
-			processTouches(&this->block, 1);
-		}
-	}
-
-	FORCEINLINE int32 GetNumHits() const
-	{
-		return Hits.Num();
-	}
-
-	FORCEINLINE HitType* GetHits()
-	{
-		return (HitType*)Hits.GetData();
-	}
-
-	FORCEINLINE const HitType* GetHits() const
-	{
-		return (const HitType*)Hits.GetData();
-	}
-};
-
-
-/** We use this struct so that if no conversion is needed in another API, we can avoid the copy (if we think that's critical) */
-struct FPhysicsRaycastInputAdapater
-{
-	FPhysicsRaycastInputAdapater(const FVector& InStart, const FVector& InDir, const EHitFlags InFlags)
-		: Start(U2PVector(InStart))
-		, Dir(U2PVector(InDir))
-		, OutputFlags(U2PHitFlags(InFlags))
-	{
-
-	}
-	PxVec3 Start;
-	PxVec3 Dir;
-	PxHitFlags OutputFlags;
-};
-
-/** We use this struct so that if no conversion is needed in another API, we can avoid the copy (if we think that's critical) */
-struct FPhysicsSweepInputAdapater
-{
-	FPhysicsSweepInputAdapater(const FTransform& InStartTM, const FVector& InDir, const EHitFlags InFlags)
-		: StartTM(U2PTransform(InStartTM))
-		, Dir(U2PVector(InDir))
-		, OutputFlags(U2PHitFlags(InFlags))
-	{
-
-	}
-	PxTransform StartTM;
-	PxVec3 Dir;
-	PxHitFlags OutputFlags;
-};
-
-/** We use this struct so that if no conversion is needed in another API, we can avoid the copy (if we think that's critical) */
-struct FPhysicsOverlapInputAdapater
-{
-	FPhysicsOverlapInputAdapater(const FTransform& InPose)
-		: GeomPose(U2PTransform(InPose))
-	{
-
-	}
-	PxTransform GeomPose;
-};
-
-FORCEINLINE uint32 GetTriangleMeshExternalFaceIndex(const PxShape& Shape, uint32 InternalFaceIndex)
-{
-	PxTriangleMeshGeometry TriangleMeshGeometry;
-	if (Shape.getTriangleMeshGeometry(TriangleMeshGeometry))
-	{
-		if (TriangleMeshGeometry.triangleMesh && InternalFaceIndex < TriangleMeshGeometry.triangleMesh->getNbTriangles())
-		{
-			if (const PxU32* TriangleRemap = TriangleMeshGeometry.triangleMesh->getTrianglesRemap())
-			{
-				return TriangleRemap[InternalFaceIndex];
-			}
-		}
-	}
-
-	return PhysXInterface::GetInvalidPhysicsFaceIndex();
-}
-
-FORCEINLINE float GetRadius(const PxCapsuleGeometry& Capsule)
-{
-	return Capsule.radius;
-}
-
-FORCEINLINE float GetHalfHeight(const PxCapsuleGeometry& Capsule)
-{
-	return Capsule.halfHeight;
-}
-
-FORCEINLINE PxTransform GetGlobalPose(const PxRigidActor& RigidActor)
-{
-	return RigidActor.getGlobalPose();
-}
-
-FORCEINLINE uint32 GetNumShapes(const PxRigidActor& RigidActor)
-{
-	return RigidActor.getNbShapes();
-}
-
-FORCEINLINE void GetShapes(const PxRigidActor& RigidActor, PxShape** ShapesBuffer, uint32 NumShapes)
-{
-	RigidActor.getShapes(ShapesBuffer, sizeof(ShapesBuffer[0]) * NumShapes);
-}
-
-FORCEINLINE void SetActor(PxActorShape& Hit, PxRigidActor* Actor)
-{
-	Hit.actor = Actor;
-}
-
-FORCEINLINE void SetShape(PxActorShape& Hit, PxShape* Shape)
-{
-	Hit.shape = Shape;
-}
-
-}
-
-#if PHYSICS_INTERFACE_PHYSX
-using namespace PhysXInterface;
-#endif
-
-#endif
\ No newline at end of file
diff --git a/Engine/Source/Runtime/PhysicsCore/Public/PhysXPublicCore.h b/Engine/Source/Runtime/PhysicsCore/Public/PhysXPublicCore.h
--- a/Engine/Source/Runtime/PhysicsCore/Public/PhysXPublicCore.h
+++ b/Engine/Source/Runtime/PhysicsCore/Public/PhysXPublicCore.h
@@ -8,311 +8,3 @@
 
 #include "CoreMinimal.h"
 #include "Stats/Stats.h"
-
-#if PHYSICS_INTERFACE_PHYSX
-#include "PhysXIncludes.h"
-#include "PhysicsInterfaceTypesCore.h"
-#include "CollisionShape.h"
-
-// Whether or not to use the PhysX scene lock
-#ifndef USE_SCENE_LOCK
-#define USE_SCENE_LOCK			1
-#endif
-
-#if USE_SCENE_LOCK
-
-DECLARE_CYCLE_STAT_EXTERN(TEXT("PhysX Scene ReadLock"), STAT_PhysSceneReadLock, STATGROUP_Physics, PHYSICSCORE_API);
-DECLARE_CYCLE_STAT_EXTERN(TEXT("PhysX Scene WriteLock"), STAT_PhysSceneWriteLock, STATGROUP_Physics, PHYSICSCORE_API);
-
-/** Scoped scene read lock - we use this instead of PxSceneReadLock because it handles NULL scene */
-class FPhysXSceneReadLock
-{
-public:
-	
-	FPhysXSceneReadLock(PxScene* PInScene, const char* filename, PxU32 lineno)
-		: PScene(PInScene)
-	{
-		SCOPE_CYCLE_COUNTER(STAT_PhysSceneReadLock);
-		if(PScene)
-		{
-			PScene->lockRead(filename, lineno);
-		}
-	}
-
-	~FPhysXSceneReadLock()
-	{
-		if(PScene)
-		{
-			PScene->unlockRead();
-		}
-	}
-
-private:
-	PxScene* PScene;
-};
-
-#if WITH_APEX
-/** Scoped scene read lock - we use this instead of PxSceneReadLock because it handles NULL scene */
-class FApexSceneReadLock
-{
-public:
-
-	FApexSceneReadLock(nvidia::apex::Scene* PInScene, const char* filename, PxU32 lineno)
-		: PScene(PInScene)
-	{
-		SCOPE_CYCLE_COUNTER(STAT_PhysSceneReadLock);
-		if (PScene)
-		{
-			PScene->lockRead(filename, lineno);
-		}
-	}
-
-	~FApexSceneReadLock()
-	{
-		if (PScene)
-		{
-			PScene->unlockRead();
-		}
-	}
-
-private:
-	nvidia::apex::Scene* PScene;
-};
-#endif
-
-/** Scoped scene write lock - we use this instead of PxSceneReadLock because it handles NULL scene */
-class FPhysXSceneWriteLock
-{
-public:
-	FPhysXSceneWriteLock(PxScene* PInScene, const char* filename, PxU32 lineno)
-		: PScene(PInScene)
-	{
-		SCOPE_CYCLE_COUNTER(STAT_PhysSceneWriteLock);
-		if(PScene)
-		{
-			PScene->lockWrite(filename, lineno);
-		}
-	}
-
-	~FPhysXSceneWriteLock()
-	{
-		if(PScene)
-		{
-			PScene->unlockWrite();
-		}
-	}
-
-private:
-	PxScene* PScene;
-};
-
-#if WITH_APEX
-/** Scoped scene write lock - we use this instead of PxSceneReadLock because it handles NULL scene */
-class FApexSceneWriteLock
-{
-public:
-	FApexSceneWriteLock(nvidia::apex::Scene* PInScene, const char* filename, PxU32 lineno)
-		: PScene(PInScene)
-	{
-		SCOPE_CYCLE_COUNTER(STAT_PhysSceneWriteLock);
-		if (PScene)
-		{
-			PScene->lockWrite(filename, lineno);
-		}
-	}
-
-	~FApexSceneWriteLock()
-	{
-		if (PScene)
-		{
-			PScene->unlockWrite();
-		}
-	}
-
-private:
-	nvidia::apex::Scene* PScene;
-};
-#endif
-
-#define SCOPED_SCENE_READ_LOCK( _scene ) FPhysXSceneReadLock PREPROCESSOR_JOIN(_rlock,__LINE__)(_scene, __FILE__, __LINE__)
-#define SCOPED_SCENE_WRITE_LOCK( _scene ) FPhysXSceneWriteLock PREPROCESSOR_JOIN(_wlock,__LINE__)(_scene, __FILE__, __LINE__)
-#if WITH_APEX
-#define SCOPED_APEX_SCENE_READ_LOCK( _scene ) FApexSceneReadLock PREPROCESSOR_JOIN(_rlock,__LINE__)(_scene, __FILE__, __LINE__)
-#define SCOPED_APEX_SCENE_WRITE_LOCK( _scene ) FApexSceneWriteLock PREPROCESSOR_JOIN(_wlock,__LINE__)(_scene, __FILE__, __LINE__)
-#endif
-
-#define SCENE_LOCK_READ( _scene )		{ SCOPE_CYCLE_COUNTER(STAT_PhysSceneReadLock); if((_scene) != NULL) { (_scene)->lockRead(__FILE__, __LINE__); } }
-#define SCENE_UNLOCK_READ( _scene )		{ if((_scene) != NULL) { (_scene)->unlockRead(); } }
-#define SCENE_LOCK_WRITE( _scene )		{ SCOPE_CYCLE_COUNTER(STAT_PhysSceneWriteLock); if((_scene) != NULL) { (_scene)->lockWrite(__FILE__, __LINE__); } }
-#define SCENE_UNLOCK_WRITE( _scene )	{ if((_scene) != NULL) { (_scene)->unlockWrite(); } }
-#else
-#define SCOPED_SCENE_READ_LOCK_INDEXED( _scene, _index )
-#define SCOPED_SCENE_READ_LOCK( _scene )
-#define SCOPED_SCENE_WRITE_LOCK_INDEXED( _scene, _index )
-#define SCOPED_SCENE_WRITE_LOCK( _scene )
-#define SCENE_LOCK_READ( _scene )
-#define SCENE_UNLOCK_READ( _scene )
-#define SCENE_LOCK_WRITE( _scene )
-#define SCENE_UNLOCK_WRITE( _scene )
-
-#if WITH_APEX
-#define SCOPED_APEX_SCENE_READ_LOCK( _scene )
-#define SCOPED_APEX_SCENE_WRITE_LOCK( _scene )
-#endif
-#endif
-
-//////// BASIC TYPE CONVERSION
-
-/** Convert Unreal FMatrix to PhysX PxTransform */
-PHYSICSCORE_API PxTransform UMatrix2PTransform(const FMatrix& UTM);
-/** Convert Unreal FMatrix to PhysX PxMat44 */
-PHYSICSCORE_API PxMat44 U2PMatrix(const FMatrix& UTM);
-/** Convert PhysX PxTransform to Unreal PxTransform */
-PHYSICSCORE_API FTransform P2UTransform(const PxTransform& PTM);
-/** Convert PhysX PxMat44 to Unreal FMatrix */
-PHYSICSCORE_API FMatrix P2UMatrix(const PxMat44& PMat);
-/** Convert PhysX PxTransform to Unreal FMatrix */
-PHYSICSCORE_API FMatrix PTransform2UMatrix(const PxTransform& PTM);
-
-// inlines
-
-/** Convert Unreal FVector to PhysX PxVec3 */
-FORCEINLINE_DEBUGGABLE PxVec3 U2PVector(const FVector& UVec)
-{
-	return PxVec3(UVec.X, UVec.Y, UVec.Z);
-}
-
-FORCEINLINE_DEBUGGABLE PxVec4 U2PVector(const FVector4& UVec)
-{
-	return PxVec4(UVec.X, UVec.Y, UVec.Z, UVec.W);
-}
-
-/** Convert Unreal FQuat to PhysX PxQuat */
-FORCEINLINE_DEBUGGABLE PxQuat U2PQuat(const FQuat& UQuat)
-{
-	return PxQuat( UQuat.X, UQuat.Y, UQuat.Z, UQuat.W );
-}
-
-/** Convert Unreal FPlane to PhysX plane def */
-FORCEINLINE_DEBUGGABLE PxPlane U2PPlane(const FPlane& Plane)
-{
-	return PxPlane(Plane.X, Plane.Y, Plane.Z, -Plane.W);
-}
-
-/** Convert PhysX PxVec3 to Unreal FVector */
-FORCEINLINE_DEBUGGABLE FVector P2UVector(const PxVec3& PVec)
-{
-	return FVector(PVec.x, PVec.y, PVec.z);
-}
-
-FORCEINLINE_DEBUGGABLE FVector4 P2UVector(const PxVec4& PVec)
-{
-	return FVector4(PVec.x, PVec.y, PVec.z, PVec.w);
-}
-
-/** Convert PhysX PxQuat to Unreal FQuat */
-FORCEINLINE_DEBUGGABLE FQuat P2UQuat(const PxQuat& PQuat)
-{
-	return FQuat(PQuat.x, PQuat.y, PQuat.z, PQuat.w);
-}
-
-/** Convert PhysX plane def to Unreal FPlane */
-FORCEINLINE_DEBUGGABLE FPlane P2UPlane(const PxReal P[4])
-{
-	return FPlane(P[0], P[1], P[2], -P[3]);
-}
-
-FORCEINLINE_DEBUGGABLE FPlane P2UPlane(const PxPlane& Plane)
-{
-	return FPlane(Plane.n.x, Plane.n.y, Plane.n.z, -Plane.d);
-}
-
-/** Convert PhysX Barycentric Vec3 to FVector4 */
-FORCEINLINE_DEBUGGABLE FVector4 P2U4BaryCoord(const PxVec3& PVec)
-{
-	return FVector4(PVec.x, PVec.y, 1.f - PVec.x - PVec.y, PVec.z);
-}
-
-/** Convert PhysX Geometry Type to ECollisionShapeType */
-inline ECollisionShapeType P2UGeometryType(PxGeometryType::Enum Type)
-{
-	switch(Type)
-	{
-	case PxGeometryType::eSPHERE: return ECollisionShapeType::Sphere;
-	case PxGeometryType::ePLANE: return ECollisionShapeType::Plane;
-	case PxGeometryType::eCAPSULE: return ECollisionShapeType::Capsule;
-	case PxGeometryType::eBOX: return ECollisionShapeType::Box;
-	case PxGeometryType::eCONVEXMESH: return ECollisionShapeType::Convex;
-	case PxGeometryType::eTRIANGLEMESH: return ECollisionShapeType::Trimesh;
-	case PxGeometryType::eHEIGHTFIELD: return ECollisionShapeType::Heightfield;
-	default: return ECollisionShapeType::None;
-	}
-}
-
-/** Convert ECollisionShapeType to PhysX Geometry Type */
-inline PxGeometryType::Enum U2PGeometryType(ECollisionShapeType Type)
-{
-	switch(Type)
-	{
-	case ECollisionShapeType::Sphere: return PxGeometryType::eSPHERE;
-	case ECollisionShapeType::Plane: return PxGeometryType::ePLANE;
-	case ECollisionShapeType::Capsule: return PxGeometryType::eCAPSULE;
-	case ECollisionShapeType::Box: return PxGeometryType::eBOX;
-	case ECollisionShapeType::Convex: return PxGeometryType::eCONVEXMESH;
-	case ECollisionShapeType::Trimesh: return PxGeometryType::eTRIANGLEMESH;
-	case ECollisionShapeType::Heightfield: return PxGeometryType::eHEIGHTFIELD;
-	default: return PxGeometryType::eINVALID;
-	}
-}
-
-/** Convert Unreal FTransform to PhysX PxTransform */
-FORCEINLINE_DEBUGGABLE PxTransform U2PTransform(const FTransform& UTransform)
-{
-	PxQuat PQuat = U2PQuat(UTransform.GetRotation());
-	PxVec3 PPos = U2PVector(UTransform.GetLocation());
-
-	PxTransform Result(PPos, PQuat);
-
-	return Result;
-}
-
-namespace nvidia
-{
-	namespace apex
-	{
-		class  PhysX3Interface;
-	}
-}
-
-/** The default interface is nullptr. This can be set by other modules to get custom behavior */
-extern PHYSICSCORE_API nvidia::apex::PhysX3Interface* GPhysX3Interface;
-
-struct PHYSICSCORE_API FContactModifyCallback : public PxContactModifyCallback
-{
-	virtual ~FContactModifyCallback() {}	//This should only be called from the factory's destroy method which is called after simulation is done.
-};
-
-struct PHYSICSCORE_API FCCDContactModifyCallback : public PxCCDContactModifyCallback
-{
-	virtual ~FCCDContactModifyCallback() {}	//This should only be called from the factory's destroy method which is called after simulation is done.
-};
-
-//////// GEOM CONVERSION
-// we need this helper struct since PhysX needs geoms to be on the stack
-struct UCollision2PGeom
-{
-	UCollision2PGeom(const FCollisionShape& CollisionShape);
-	const PxGeometry * GetGeometry() { return (PxGeometry*)Storage; }
-private:
-
-	union StorageUnion
-	{
-		char box[sizeof(PxBoxGeometry)];
-		char sphere[sizeof(PxSphereGeometry)];
-		char capsule[sizeof(PxCapsuleGeometry)];
-	};	//we need this because we can't use real union since these structs have non trivial constructors
-
-	char Storage[sizeof(StorageUnion)];
-};
-
-#endif
\ No newline at end of file
diff --git a/Engine/Source/Runtime/PhysicsCore/Public/PhysXSupportCore.h b/Engine/Source/Runtime/PhysicsCore/Public/PhysXSupportCore.h
--- a/Engine/Source/Runtime/PhysicsCore/Public/PhysXSupportCore.h
+++ b/Engine/Source/Runtime/PhysicsCore/Public/PhysXSupportCore.h
@@ -9,437 +9,3 @@
 #include "CoreMinimal.h"
 #include "PhysicsCore.h"
 #include "PhysicsPublicCore.h"
-
-#if PHYSICS_INTERFACE_PHYSX
-#include "PhysXPublicCore.h"
-#include "HAL/IConsoleManager.h"
-
-/////// GLOBAL POINTERS
-/** Pointer to PhysX Foundation singleton */
-extern PHYSICSCORE_API physx::PxFoundation* GPhysXFoundation;
-/** Pointer to PhysX debugger */
-extern PHYSICSCORE_API PxPvd* GPhysXVisualDebugger;
-
-// Whether to track PhysX memory allocations
-#ifndef PHYSX_MEMORY_VALIDATION
-#define PHYSX_MEMORY_VALIDATION		0
-#endif
-
-// Whether to track PhysX memory allocations
-#ifndef PHYSX_MEMORY_STATS
-#define PHYSX_MEMORY_STATS		0 || PHYSX_MEMORY_VALIDATION
-#endif
-
-#define PHYSX_MEMORY_STAT_ONLY (0)
-
-/** PhysX memory allocator wrapper */
-class PHYSICSCORE_API FPhysXAllocator : public PxAllocatorCallback
-{
-#if PHYSX_MEMORY_STATS
-	TMap<FName, size_t> AllocationsByType;
-
-	struct FPhysXAllocationHeader
-	{
-		FPhysXAllocationHeader(FName InAllocationTypeName, size_t InAllocationSize)
-			: AllocationTypeName(InAllocationTypeName)
-			, AllocationSize(InAllocationSize)
-		{
-			static_assert((sizeof(FPhysXAllocationHeader) % 16) == 0, "FPhysXAllocationHeader size must multiple of bytes.");
-			MagicPadding();
-		}
-
-		void MagicPadding()
-		{
-			for (uint8 ByteCount = 0; ByteCount < sizeof(Padding); ++ByteCount)
-			{
-				Padding[ByteCount] = 'A' + ByteCount % 4;
-			}
-		}
-
-		bool operator==(const FPhysXAllocationHeader& OtherHeader) const
-		{
-			bool bHeaderSame = AllocationTypeName == OtherHeader.AllocationTypeName && AllocationSize == OtherHeader.AllocationSize;
-			for (uint8 ByteCount = 0; ByteCount < sizeof(Padding); ++ByteCount)
-			{
-				bHeaderSame &= Padding[ByteCount] == OtherHeader.Padding[ByteCount];
-			}
-
-			return bHeaderSame;
-		}
-
-		FName AllocationTypeName;
-		size_t	AllocationSize;
-		static const int PaddingSize = 8;
-		uint8 Padding[PaddingSize];	//physx needs 16 byte alignment. Additionally we fill padding with a pattern to see if there's any memory stomps
-		uint8 Padding2[(sizeof(FName) + sizeof(size_t) + PaddingSize) % 16];
-
-		void Validate() const
-		{
-			bool bValid = true;
-			for (uint8 ByteCount = 0; ByteCount < sizeof(Padding); ++ByteCount)
-			{
-				bValid &= Padding[ByteCount] == 'A' + ByteCount % 4;
-			}
-
-			check(bValid);
-
-			FPhysXAllocationHeader* AllocationFooter = (FPhysXAllocationHeader*)(((uint8*)this) + sizeof(FPhysXAllocationHeader) + AllocationSize);
-			check(*AllocationFooter == *this);
-		}
-	};
-
-#endif
-
-public:
-
-#if PHYSX_MEMORY_VALIDATION
-
-	/** Iterates over all allocations and checks that they the headers and footers are valid */
-	void ValidateHeaders()
-	{
-		check(IsInGameThread());
-		FPhysXAllocationHeader* TmpHeader = nullptr;
-		while (NewHeaders.Dequeue(TmpHeader))
-		{
-			AllocatedHeaders.Add(TmpHeader);
-		}
-
-		while (OldHeaders.Dequeue(TmpHeader))
-		{
-			AllocatedHeaders.Remove(TmpHeader);
-		}
-
-		FScopeLock Lock(&ValidationCS);	//this is needed in case another thread is freeing the header
-		for (FPhysXAllocationHeader* Header : AllocatedHeaders)
-		{
-			Header->Validate();
-		}
-	}
-#endif
-
-	FPhysXAllocator()
-	{}
-
-	virtual ~FPhysXAllocator()
-	{}
-
-	virtual void* allocate(size_t size, const char* typeName, const char* filename, int line) override
-	{
-#if PHYSX_MEMORY_STATS
-		INC_DWORD_STAT_BY(STAT_MemoryPhysXTotalAllocationSize, size);
-
-
-		FString AllocationString = FString::Printf(TEXT("%s %s:%d"), ANSI_TO_TCHAR(typeName), ANSI_TO_TCHAR(filename), line);
-		FName AllocationName(*AllocationString);
-
-		// Assign header
-		FPhysXAllocationHeader* AllocationHeader = (FPhysXAllocationHeader*)FMemory::Malloc(size + sizeof(FPhysXAllocationHeader) * 2, 16);
-		AllocationHeader->AllocationTypeName = AllocationName;
-		AllocationHeader->AllocationSize = size;
-		AllocationHeader->MagicPadding();
-		FPhysXAllocationHeader* AllocationFooter = (FPhysXAllocationHeader*)(((uint8*)AllocationHeader) + size + sizeof(FPhysXAllocationHeader));
-		AllocationFooter->AllocationTypeName = AllocationName;
-		AllocationFooter->AllocationSize = size;
-		AllocationFooter->MagicPadding();
-
-		size_t* TotalByType = AllocationsByType.Find(AllocationName);	//TODO: this is not thread safe!
-		if (TotalByType)
-		{
-			*TotalByType += size;
-		}
-		else
-		{
-			AllocationsByType.Add(AllocationName, size);
-		}
-
-#if PHYSX_MEMORY_VALIDATION
-		NewHeaders.Enqueue(AllocationHeader);
-#endif
-
-		return (uint8*)AllocationHeader + sizeof(FPhysXAllocationHeader);
-#else
-		LLM_SCOPE(ELLMTag::PhysXAllocator);
-		void* ptr = FMemory::Malloc(size, 16);
-#if PHYSX_MEMORY_STAT_ONLY
-		INC_DWORD_STAT_BY(STAT_MemoryPhysXTotalAllocationSize, FMemory::GetAllocSize(ptr));
-#endif
-		return ptr;
-#endif
-	}
-
-	virtual void deallocate(void* ptr) override
-	{
-#if PHYSX_MEMORY_STATS
-		if (ptr)
-		{
-			FPhysXAllocationHeader* AllocationHeader = (FPhysXAllocationHeader*)((uint8*)ptr - sizeof(FPhysXAllocationHeader));
-#if PHYSX_MEMORY_VALIDATION
-			AllocationHeader->Validate();
-			OldHeaders.Enqueue(AllocationHeader);
-			FScopeLock Lock(&ValidationCS);	//this is needed in case we are in the middle of validating the headers
-#endif
-
-			DEC_DWORD_STAT_BY(STAT_MemoryPhysXTotalAllocationSize, AllocationHeader->AllocationSize);
-			size_t* TotalByType = AllocationsByType.Find(AllocationHeader->AllocationTypeName);
-			*TotalByType -= AllocationHeader->AllocationSize;
-			FMemory::Free(AllocationHeader);
-		}
-#else
-#if PHYSX_MEMORY_STAT_ONLY
-		DEC_DWORD_STAT_BY(STAT_MemoryPhysXTotalAllocationSize, FMemory::GetAllocSize(ptr));
-#endif
-		FMemory::Free(ptr);
-#endif
-	}
-
-#if PHYSX_MEMORY_STATS
-	void DumpAllocations(FOutputDevice* Ar)
-	{
-		struct FSortBySize
-		{
-			FORCEINLINE bool operator()(const size_t& A, const size_t& B) const
-			{
-				// Sort descending
-				return B < A;
-			}
-		};
-
-		size_t TotalSize = 0;
-		AllocationsByType.ValueSort(FSortBySize());
-		for (auto It = AllocationsByType.CreateConstIterator(); It; ++It)
-		{
-			TotalSize += It.Value();
-			Ar->Logf(TEXT("%-10d %s"), It.Value(), *It.Key().ToString());
-		}
-
-		Ar->Logf(TEXT("Total:%-10d"), TotalSize);
-	}
-#endif
-
-#if PHYSX_MEMORY_VALIDATION
-private:
-	FCriticalSection ValidationCS;
-	TSet<FPhysXAllocationHeader*> AllocatedHeaders;
-
-	//Since this needs to be thread safe we can't add to the allocated headers set until we're on the game thread
-	TQueue<FPhysXAllocationHeader*, EQueueMode::Mpsc> NewHeaders;
-	TQueue<FPhysXAllocationHeader*, EQueueMode::Mpsc> OldHeaders;
-
-#endif
-};
-
-extern PHYSICSCORE_API int32 GPhysXHackCurrentLoopCounter;
-
-/** PhysX output stream wrapper */
-class PHYSICSCORE_API FPhysXErrorCallback : public PxErrorCallback
-{
-public:
-	virtual void reportError(PxErrorCode::Enum e, const char* message, const char* file, int line) override;
-};
-
-extern PHYSICSCORE_API TAutoConsoleVariable<float> CVarToleranceScaleLength;
-extern PHYSICSCORE_API TAutoConsoleVariable<float> CVarToleranceScaleSpeed;
-
-
-/** Utility class to keep track of shared physics data */
-class PHYSICSCORE_API FPhysxSharedData
-{
-public:
-	static FPhysxSharedData& Get() { return *Singleton; }
-	static void Initialize();
-	static void Terminate();
-
-	static void LockAccess();
-	static void UnlockAccess();
-
-	void Add(PxBase* Obj, const FString& OwnerName);
-	void Remove(PxBase* Obj);
-
-	const PxCollection* GetCollection() { return SharedObjects; }
-
-	void DumpSharedMemoryUsage(FOutputDevice* Ar);
-private:
-	/** Collection of shared physx objects */
-	PxCollection* SharedObjects;
-	TMap<PxBase*, FString> OwnerNames;
-	FCriticalSection CriticalSection;
-
-	static FPhysxSharedData* Singleton;
-
-	FPhysxSharedData()
-	{
-		SharedObjects = PxCreateCollection();
-	}
-
-	~FPhysxSharedData()
-	{
-		SharedObjects->release();
-	}
-
-};
-
-
-/** Utility wrapper for a PhysX output stream that only counts the memory. */
-class PHYSICSCORE_API FPhysXCountMemoryStream : public PxOutputStream
-{
-public:
-	/** Memory used by the serialized object(s) */
-	uint32 UsedMemory;
-
-	FPhysXCountMemoryStream()
-		: UsedMemory(0)
-	{}
-
-	virtual PxU32 write(const void* Src, PxU32 Count) override
-	{
-		UsedMemory += Count;
-		return Count;
-	}
-};
-
-/** Utility wrapper for a uint8 TArray for saving into PhysX. */
-class PHYSICSCORE_API FPhysXOutputStream : public PxOutputStream
-{
-public:
-	/** Raw byte data */
-	TArray<uint8>			*Data;
-
-	FPhysXOutputStream()
-		: Data(NULL)
-	{}
-
-	FPhysXOutputStream(TArray<uint8> *InData)
-		: Data(InData)
-	{}
-
-	PxU32 write(const void* Src, PxU32 Count)
-	{
-		check(Data);
-		if (Count)	//PhysX serializer can pass us 0 bytes to write
-		{
-			check(Src);
-			int32 CurrentNum = (*Data).Num();
-			(*Data).AddUninitialized(Count);
-			FMemory::Memcpy(&(*Data)[CurrentNum], Src, Count);
-		}
-
-		return Count;
-	}
-};
-
-/**
- * Returns the in-memory size of the specified object by serializing it.
- *
- * @param	Obj					Object to determine the memory footprint for
- * @param	SharedCollection	Shared collection of data to ignore
- * @returns						Size of the object in bytes determined by serialization
- **/
-PHYSICSCORE_API SIZE_T GetPhysxObjectSize(PxBase* Obj, const PxCollection* SharedCollection);
-
-PHYSICSCORE_API void PvdConnect(FString Host, bool bVisualization);
-
-
-#if WITH_APEX
-/**
-	"Null" render resource manager callback for APEX
-	This just gives a trivial implementation of the interface, since we are not using the APEX rendering API
-*/
-class PHYSICSCORE_API FApexNullRenderResourceManager : public nvidia::apex::UserRenderResourceManager
-{
-public:
-	// NxUserRenderResourceManager interface.
-
-	virtual nvidia::apex::UserRenderVertexBuffer*	createVertexBuffer(const nvidia::apex::UserRenderVertexBufferDesc&) override
-	{
-		return NULL;
-	}
-	virtual nvidia::apex::UserRenderIndexBuffer*	createIndexBuffer(const nvidia::apex::UserRenderIndexBufferDesc&) override
-	{
-		return NULL;
-	}
-	virtual nvidia::apex::UserRenderBoneBuffer*		createBoneBuffer(const nvidia::apex::UserRenderBoneBufferDesc&) override
-	{
-		return NULL;
-	}
-	virtual nvidia::apex::UserRenderInstanceBuffer*	createInstanceBuffer(const nvidia::apex::UserRenderInstanceBufferDesc&) override
-	{
-		return NULL;
-	}
-	virtual nvidia::apex::UserRenderSpriteBuffer*   createSpriteBuffer(const nvidia::apex::UserRenderSpriteBufferDesc&) override
-	{
-		return NULL;
-	}
-
-	virtual nvidia::apex::UserRenderSurfaceBuffer*  createSurfaceBuffer(const nvidia::apex::UserRenderSurfaceBufferDesc& desc)   override
-	{
-		return NULL;
-	}
-
-	virtual nvidia::apex::UserRenderResource*		createResource(const nvidia::apex::UserRenderResourceDesc&) override
-	{
-		return NULL;
-	}
-	virtual void						releaseVertexBuffer(nvidia::apex::UserRenderVertexBuffer&) override {}
-	virtual void						releaseIndexBuffer(nvidia::apex::UserRenderIndexBuffer&) override {}
-	virtual void						releaseBoneBuffer(nvidia::apex::UserRenderBoneBuffer&) override {}
-	virtual void						releaseInstanceBuffer(nvidia::apex::UserRenderInstanceBuffer&) override {}
-	virtual void						releaseSpriteBuffer(nvidia::apex::UserRenderSpriteBuffer&) override {}
-	virtual void                        releaseSurfaceBuffer(nvidia::apex::UserRenderSurfaceBuffer& buffer) override {}
-	virtual void						releaseResource(nvidia::apex::UserRenderResource&) override {}
-
-	virtual physx::PxU32				getMaxBonesForMaterial(void*) override
-	{
-		return 0;
-	}
-	virtual bool						getSpriteLayoutData(physx::PxU32, physx::PxU32, nvidia::apex::UserRenderSpriteBufferDesc*) override
-	{
-		return false;
-	}
-	virtual bool						getInstanceLayoutData(physx::PxU32, physx::PxU32, nvidia::apex::UserRenderInstanceBufferDesc*) override
-	{
-		return false;
-	}
-
-};
-extern PHYSICSCORE_API FApexNullRenderResourceManager GApexNullRenderResourceManager;
-
-/**
-	APEX resource callback
-	The resource callback is how APEX asks the application to find assets when it needs them
-*/
-class PHYSICSCORE_API FApexResourceCallback : public nvidia::apex::ResourceCallback
-{
-public:
-	// NxResourceCallback interface.
-
-	virtual void* requestResource(const char* NameSpace, const char* Name) override
-	{
-		// Here a pointer is looked up by name and returned
-		(void)NameSpace;
-		(void)Name;
-
-		return NULL;
-	}
-
-	virtual void  releaseResource(const char* NameSpace, const char* Name, void* Resource) override
-	{
-		// Here we release a named resource
-		(void)NameSpace;
-		(void)Name;
-		(void)Resource;
-	}
-};
-extern PHYSICSCORE_API FApexResourceCallback GApexResourceCallback;
-
-
-#endif // #if WITH_APEX
-
-inline PxSceneDesc CreateDummyPhysXSceneDescriptor()
-{
-	PxSceneDesc Desc(GPhysXSDK->getTolerancesScale());
-	Desc.filterShader = PxDefaultSimulationFilterShader;
-	Desc.cpuDispatcher = PxDefaultCpuDispatcherCreate(4);
-	return Desc;
-}
-
-#endif
\ No newline at end of file
