Physics interface cleanup.

diff --git a/Engine/Source/Runtime/PhysicsCore/Public/PhysXToChaosUtil.h b/Engine/Source/Runtime/PhysicsCore/Public/PhysXToChaosUtil.h
--- a/Engine/Source/Runtime/PhysicsCore/Public/PhysXToChaosUtil.h
+++ b/Engine/Source/Runtime/PhysicsCore/Public/PhysXToChaosUtil.h
@@ -1,202 +1,3 @@
 // Copyright Epic Games, Inc. All Rights Reserved.
 
 #pragma once 
-
-#if 0
-#include "CoreMinimal.h"
-
-#if WITH_PHYSX
-#include "PhysXIncludes.h"
-#include "PhysXSupportCore.h"
-
-#include "Chaos/Particles.h"
-#include "Chaos/Box.h"
-#include "Chaos/Sphere.h"
-#include "Chaos/Capsule.h"
-#include "Chaos/ImplicitObjectTransformed.h"
-#include "Chaos/Convex.h"
-#include "Chaos/TriangleMeshImplicitObject.h"
-#include "Chaos/HeightField.h"
-#include "Chaos/ImplicitObjectScaled.h"
-
-//convert physx shape to chaos geometry - must be inline because of unique ptr crossing dlls
-inline TUniquePtr<Chaos::TImplicitObjectTransformed<FReal, 3>> PxShapeToChaosGeom(PxShape* Shape)
-{
-	using namespace Chaos;
-
-	TUniquePtr<FImplicitObject> InnerObj;
-	PxTransform ShapeTM = Shape->getLocalPose();
-	PxGeometryHolder Geom = Shape->getGeometry();
-	switch (Geom.getType())
-	{
-	case PxGeometryType::eSPHERE:
-	{
-		InnerObj = MakeUnique<TSphere<FReal, 3>>(FVec3(0), Geom.sphere().radius);
-		break;
-	}
-	case PxGeometryType::eBOX:
-	{
-		const FVec3 HalfExtents = P2UVector(Geom.box().halfExtents);
-		InnerObj = MakeUnique<TBox<FReal, 3>>(-HalfExtents, HalfExtents);
-		break;
-	}
-	case PxGeometryType::eCAPSULE:
-	{
-		const FReal HalfHeight = Geom.capsule().halfHeight;
-		const FReal Radius = Geom.capsule().radius;
-		const FVec3 Top(HalfHeight, 0, 0);	//PhysX capsules are aligned along the x-axis
-		const FVec3 Bottom = -Top;
-		InnerObj = MakeUnique<FCapsule>(Top, Bottom, Radius);
-		break;
-	}
-	case PxGeometryType::eCONVEXMESH:
-	{
-		//todo: instancing support
-		const PxConvexMeshGeometry& ConvexMeshGeom = Geom.convexMesh();
-		const PxConvexMesh* ConvexMesh = ConvexMeshGeom.convexMesh;
-		const PxVec3* Vertices = ConvexMesh->getVertices();
-		const uint32 NumVerts = ConvexMesh->getNbVertices();
-
-		FParticles Particles;
-		Particles.AddParticles(NumVerts);
-		for (uint32 Idx = 0; Idx < NumVerts; ++Idx)
-		{
-			Particles.X(Idx) = P2UVector(Vertices[Idx]);
-		}
-
-		TUniquePtr<FConvex> ConvexObj = MakeUnique<FConvex>(Particles);
-		if (ConvexMeshGeom.scale.isIdentity())
-		{
-			InnerObj = MoveTemp(ConvexObj);
-		}
-		else
-		{
-			ensure(ConvexMeshGeom.scale.rotation == PxQuat(PxIdentity));
-			InnerObj = MakeUnique<TImplicitObjectScaled<FConvex, /*bInstanced=*/false>>(MoveTemp(ConvexObj), P2UVector(ConvexMeshGeom.scale.scale));	//todo: make instanced
-		}
-		break;
-	}
-	case PxGeometryType::eHEIGHTFIELD:
-	{
-		const PxHeightFieldGeometry& HeightFieldGeom = Geom.heightField();
-		const PxHeightField* HeightField = HeightFieldGeom.heightField;
-		const int32 NumRows = HeightField->getNbRows();
-		const int32 NumCols = HeightField->getNbColumns();
-		const int32 NumCells = NumRows * NumCols;
-
-		TArray<PxHeightFieldSample> CellBuffer;
-		CellBuffer.AddUninitialized(NumCells);
-
-		HeightField->saveCells(CellBuffer.GetData(), sizeof(CellBuffer[0]) * CellBuffer.Num());
-
-		const int32 NumTris = (NumCells - NumRows - NumCols + 1) * 2;
-		TArray<uint8> MaterialIndices;
-		MaterialIndices.Reserve(NumTris / 2);
-		for(int32 TriangleIndex = 0; TriangleIndex < NumTris; TriangleIndex += 2)
-		{
-			// We're only grabbing every other material because Unreal will only use one material per cell
-			MaterialIndices.Add((uint8)(HeightField->getTriangleMaterialIndex(TriangleIndex)));
-		}
-		
-		TArray<float> Height;
-		Height.AddUninitialized(NumRows * NumCols);
-		//PhysX and unreal have opposite handedness so we flip the data (see LandscapeComponent)
-		for (int32 RowIdx = 0; RowIdx < NumRows; ++RowIdx)
-		{
-			for (int32 ColIdx = 0; ColIdx < NumCols; ++ColIdx)
-			{
-				int32 FlippedRowIdx = NumRows - RowIdx - 1;
-				Height[ColIdx * NumRows + RowIdx] = CellBuffer[FlippedRowIdx * NumCols + ColIdx].height;	//swap rows and columns because of how physx stores it vs chaos
-			}
-		}
-
-		// Now we fix the rotation (x,y,z) becomes (x,z,y)
-		// Note we are assuming all heightfields come from LandscapeComponent
-		//FTransform LandscapeComponentTransform = P2UTransform(ShapeTM);
-		FTransform LandscapeComponentTransform = FTransform::Identity;
-		FMatrix LandscapeComponentMatrix = LandscapeComponentTransform.ToMatrixWithScale();
-
-		// Reorder the axes
-		FVector TerrainX = LandscapeComponentMatrix.GetScaledAxis(EAxis::X);
-		FVector TerrainY = LandscapeComponentMatrix.GetScaledAxis(EAxis::Y);
-		FVector TerrainZ = LandscapeComponentMatrix.GetScaledAxis(EAxis::Z);
-		LandscapeComponentMatrix.SetAxis(0, TerrainX);
-		LandscapeComponentMatrix.SetAxis(2, TerrainY);
-		LandscapeComponentMatrix.SetAxis(1, TerrainZ);
-
-		ShapeTM = U2PTransform(FTransform(LandscapeComponentMatrix));
-
-		InnerObj = MakeUnique<FHeightField>(MoveTemp(Height), MoveTemp(MaterialIndices), NumRows, NumCols, FVec3(HeightFieldGeom.columnScale, HeightFieldGeom.rowScale, HeightFieldGeom.heightScale));
-		break;
-	}
-	case PxGeometryType::eTRIANGLEMESH:
-	{
-		//todo: instancing support
-		const PxTriangleMeshGeometry& TriMeshGeom = Geom.triangleMesh();
-		const PxTriangleMesh* TriangleMesh = TriMeshGeom.triangleMesh;
-		const PxVec3* Vertices = TriangleMesh->getVertices();
-		const uint32 NumVerts = TriangleMesh->getNbVertices();
-
-		FParticles Particles;
-		Particles.AddParticles(NumVerts);
-		for (uint32 Idx = 0; Idx < NumVerts; ++Idx)
-		{
-			Particles.X(Idx) = P2UVector(Vertices[Idx]);
-		}
-
-		const void* IndexBuffer = TriangleMesh->getTriangles();
-		const uint32 NumTriangles = TriangleMesh->getNbTriangles();
-		PxTriangleMeshFlags Flags = TriangleMesh->getTriangleMeshFlags();
-		TArray<TVec3<int32>> Triangles;
-		TArray<uint16> MaterialIndices;
-		Triangles.AddUninitialized(NumTriangles);
-		MaterialIndices.Reserve(NumTriangles);
-
-		const bool bHasMaterials = NumTriangles > 0 && TriangleMesh->getTriangleMaterialIndex(0) != TNumericLimits<uint16>::Max();
-
-		for(uint32 TriIdx = 0; TriIdx < NumTriangles; ++TriIdx)
-		{
-			if (Flags & PxTriangleMeshFlag::e16_BIT_INDICES)
-			{
-				TVec3<int32> Triangle = {
-					static_cast<const uint16*>(IndexBuffer)[TriIdx * 3],
-					static_cast<const uint16*>(IndexBuffer)[TriIdx * 3+1],
-					static_cast<const uint16*>(IndexBuffer)[TriIdx * 3+2] };
-				Triangles[TriIdx] = Triangle;
-			}
-			else
-			{
-				TVec3<int32> Triangle = {
-					static_cast<const int32*>(IndexBuffer)[TriIdx * 3],
-					static_cast<const int32*>(IndexBuffer)[TriIdx * 3+1],
-					static_cast<const int32*>(IndexBuffer)[TriIdx * 3+2] };
-				Triangles[TriIdx] = Triangle;
-			}
-
-			if(bHasMaterials)
-			{
-				MaterialIndices.Add(TriangleMesh->getTriangleMaterialIndex(TriIdx));
-			}
-		}
-
-		TUniquePtr<FTriangleMeshImplicitObject> TriMeshObj = MakeUnique<FTriangleMeshImplicitObject>(MoveTemp(Particles), MoveTemp(Triangles), MoveTemp(MaterialIndices));
-		if (TriMeshGeom.scale.isIdentity())
-		{
-			InnerObj = MoveTemp(TriMeshObj);
-		}
-		else
-		{
-			ensure(TriMeshGeom.scale.rotation == PxQuat(PxIdentity));
-			InnerObj = MakeUnique<TImplicitObjectScaled<FTriangleMeshImplicitObject, /*bInstanced=*/false>>(MoveTemp(TriMeshObj), P2UVector(TriMeshGeom.scale.scale));	//todo: make instanced
-		}
-		break;
-	}
-	default: ensure(false); return nullptr;	//missing support for this geometry type
-	}
-
-	return MakeUnique<TImplicitObjectTransformed<FReal, 3>>(MoveTemp(InnerObj), P2UTransform(ShapeTM));
-}
-
-#endif
-
-#endif
\ No newline at end of file
diff --git a/Engine/Source/Runtime/PhysicsCore/Public/PhysicalMaterials/PhysicalMaterial.h b/Engine/Source/Runtime/PhysicsCore/Public/PhysicalMaterials/PhysicalMaterial.h
--- a/Engine/Source/Runtime/PhysicsCore/Public/PhysicalMaterials/PhysicalMaterial.h
+++ b/Engine/Source/Runtime/PhysicsCore/Public/PhysicalMaterials/PhysicalMaterial.h
@@ -11,13 +11,6 @@
 
 struct FPropertyChangedEvent;
 
-#if PHYSICS_INTERFACE_PHYSX
-namespace physx
-{
-	class PxMaterial;
-}
-#endif
-
 /**
  * Physical materials are used to define the response of a physical object when interacting dynamically with the world.
  */
diff --git a/Engine/Source/Runtime/PhysicsCore/Public/PhysicsInitialization.h b/Engine/Source/Runtime/PhysicsCore/Public/PhysicsInitialization.h
--- a/Engine/Source/Runtime/PhysicsCore/Public/PhysicsInitialization.h
+++ b/Engine/Source/Runtime/PhysicsCore/Public/PhysicsInitialization.h
@@ -12,12 +12,6 @@ namespace PhysDLLHelper
 	 */
 	PHYSICSCORE_API bool LoadPhysXModules(bool bLoadCooking);
 
-
-#if WITH_APEX
-	PHYSICSCORE_API void* LoadAPEXModule(const FString& Path);
-	PHYSICSCORE_API void UnloadAPEXModule(void* Handle);
-#endif
-
 	/**
 	 *	Unload the required modules for PhysX
 	 */
diff --git a/Engine/Source/Runtime/PhysicsCore/Public/PhysicsInterfaceDeclaresCore.h b/Engine/Source/Runtime/PhysicsCore/Public/PhysicsInterfaceDeclaresCore.h
--- a/Engine/Source/Runtime/PhysicsCore/Public/PhysicsInterfaceDeclaresCore.h
+++ b/Engine/Source/Runtime/PhysicsCore/Public/PhysicsInterfaceDeclaresCore.h
@@ -4,24 +4,6 @@
 
 #define MAX_PHYSICS_FIELD_TARGETS 32
 
-#ifndef PHYSICS_INTERFACE_PHYSX
-#define PHYSICS_INTERFACE_PHYSX 0
-#endif
-
-#ifndef WITH_CHAOS
-#define WITH_CHAOS 0
-#endif
-
-#ifndef WITH_PHYSX
-#define WITH_PHYSX 0
-#endif
-
-#ifndef WITH_IMMEDIATE_PHYSX
-#define WITH_IMMEDIATE_PHYSX 0
-#endif
-
-#if WITH_CHAOS
-
 #include "ChaosSQTypes.h"
 #include "PhysicsProxy/SingleParticlePhysicsProxyFwd.h"
 
@@ -91,124 +73,6 @@ using FSingleHitBuffer = ChaosInterface::FSQSingleHitBuffer<T>;
 template <typename T>
 using FDynamicHitBuffer = ChaosInterface::FSQHitBuffer<T>;
 
-#elif PHYSICS_INTERFACE_PHYSX
-
-namespace physx
-{
-	struct PxLocationHit;
-	struct PxSweepHit;
-	struct PxRaycastHit;
-	struct PxOverlapHit;
-	struct PxQueryFlag;
-	struct PxQueryHit;
-
-	class PxTransform;
-	class PxShape;
-	class PxGeometry;
-	class PxCapsuleGeometry;
-	class PxMaterial;
-	class PxRigidActor;
-
-	template<typename T>
-	struct PxHitBuffer;
-	
-	template<typename T>
-	struct PxHitCallback;
-}
-
-namespace PhysXInterface
-{
-	struct FDummyPhysType {};
-}
-
-using FHitLocation = physx::PxLocationHit;
-using FHitSweep = physx::PxSweepHit;
-using FHitRaycast = physx::PxRaycastHit;
-using FHitOverlap = physx::PxOverlapHit;
-using FPhysicsQueryHit = physx::PxQueryHit;
-
-using FPhysicsTransform = physx::PxTransform;
-using FPhysTypeDummy = PhysXInterface::FDummyPhysType;
-
-using FPhysicsShape = physx::PxShape;
-using FPhysicsGeometry = physx::PxGeometry;
-using FPhysicsCapsuleGeometry = physx::PxCapsuleGeometry;
-using FPhysicsMaterial = physx::PxMaterial;
-using FPhysicsActor = physx::PxRigidActor;
-
-template <typename T>
-using FPhysicsHitCallback = physx::PxHitCallback<T>;
-
-struct FQueryDebugParams {};
-
-#else
-
-static_assert(false, "A physics engine interface must be defined to build");
-
-#endif
-
-#if PHYSICS_INTERFACE_PHYSX
-
-#if WITH_IMMEDIATE_PHYSX
-
-struct FPhysicsActorReference_ImmediatePhysX;
-struct FPhysicsConstraintReference_ImmediatePhysX;
-struct FPhysicsInterface_ImmediatePhysX;
-class FPhysScene_ImmediatePhysX;
-struct FPhysicsAggregateReference_ImmediatePhysX;
-struct FPhysicsCommand_ImmediatePhysX;
-struct FPhysicsShapeReference_ImmediatePhysX;
-struct FPhysicsMaterialReference_ImmediatePhysX;
-struct FPhysicsGeometryCollection_ImmediatePhysX;
-struct FPhysXShapeAdapter;
-struct FPhysxUserData;
-
-typedef FPhysicsActorReference_ImmediatePhysX		FPhysicsActorHandle;
-typedef FPhysicsConstraintReference_ImmediatePhysX	FPhysicsConstraintHandle;
-typedef FPhysicsInterface_ImmediatePhysX			FPhysicsInterface;
-typedef FPhysScene_ImmediatePhysX					FPhysScene;
-typedef FPhysicsAggregateReference_ImmediatePhysX	FPhysicsAggregateHandle;
-typedef FPhysicsCommand_ImmediatePhysX				FPhysicsCommand;
-typedef FPhysicsShapeReference_ImmediatePhysX		FPhysicsShapeHandle;
-typedef FPhysicsGeometryCollection_ImmediatePhysX	FPhysicsGeometryCollection;
-typedef FPhysicsMaterialReference_ImmediatePhysX	FPhysicsMaterialHandle;
-typedef FPhysXShapeAdapter							FPhysicsShapeAdapter;
-typedef FPhysxUserData								FPhysicsUserData;
-
-inline FPhysicsActorHandle DefaultPhysicsActorHandle() { return FPhysicsActorHandle(); }
-
-#else
-
-struct FPhysicsActorHandle_PhysX;
-struct FPhysicsConstraintHandle_PhysX;
-struct FPhysicsInterface_PhysX;
-class FPhysScene_PhysX;
-struct FPhysicsAggregateHandle_PhysX;
-struct FPhysicsCommand_PhysX;
-struct FPhysicsShapeHandle_PhysX;
-struct FPhysicsGeometryCollection_PhysX;
-struct FPhysicsMaterialHandle_PhysX;
-struct FPhysXShapeAdapter;
-struct FPhysxUserData;
-
-typedef FPhysicsActorHandle_PhysX			FPhysicsActorHandle;
-typedef FPhysicsConstraintHandle_PhysX		FPhysicsConstraintHandle;
-typedef FPhysicsInterface_PhysX				FPhysicsInterface;
-typedef FPhysScene_PhysX					FPhysScene;
-typedef FPhysicsAggregateHandle_PhysX		FPhysicsAggregateHandle;
-typedef FPhysicsCommand_PhysX				FPhysicsCommand;
-typedef FPhysicsShapeHandle_PhysX			FPhysicsShapeHandle;
-typedef FPhysicsGeometryCollection_PhysX	FPhysicsGeometryCollection;
-typedef FPhysicsMaterialHandle_PhysX		FPhysicsMaterialHandle;
-typedef FPhysXShapeAdapter					FPhysicsShapeAdapter;
-typedef FPhysxUserData						FPhysicsUserData;
-
-extern FPhysicsActorHandle DefaultPhysicsActorHandle();
-
-#endif
-
-#elif WITH_CHAOS
-
 using FPhysicsActorHandle = Chaos::FSingleParticlePhysicsProxy*;
 
 class FChaosSceneId;
@@ -234,9 +98,3 @@ typedef FPhysicsShapeAdapter_Chaos			FPhysicsShapeAdapter;
 typedef FPhysicsUserData_Chaos				FPhysicsUserData;
 
 inline FPhysicsActorHandle DefaultPhysicsActorHandle() { return nullptr; }
-
-#else
-
-static_assert(false, "A physics engine interface must be defined to build");
-
-#endif
diff --git a/Engine/Source/Runtime/PhysicsCore/Public/PhysicsInterfaceTypesCore.h b/Engine/Source/Runtime/PhysicsCore/Public/PhysicsInterfaceTypesCore.h
--- a/Engine/Source/Runtime/PhysicsCore/Public/PhysicsInterfaceTypesCore.h
+++ b/Engine/Source/Runtime/PhysicsCore/Public/PhysicsInterfaceTypesCore.h
@@ -22,12 +22,7 @@ struct FActorCreationParams
 		, bStartAwake(true)
 		, DebugName(nullptr)
 	{}
-
-#if WITH_CHAOS
 	FChaosScene* Scene;
-#else
-	FPhysScene* Scene;
-#endif
 	FTransform InitialTM;
 	bool bStatic;
 	bool bQueryOnly;
@@ -164,9 +159,7 @@ static bool CalcMeshNegScaleCompensation(const FVector& InScale3D, FTransform& O
 // TODO: Fixup types, these are more or less the PhysX types renamed as a temporary solution.
 // Probably should move to different header as well.
 
-#if !PHYSICS_INTERFACE_PHYSX
 const uint32 AggregateMaxSize = 128;
-#endif
 
 class UPhysicalMaterial;
 class UPrimitiveComponent;
@@ -217,14 +210,8 @@ public:
 	static bool IsGarbage(void* UserData){ return ((FChaosUserData*)UserData)->Type < EChaosUserDataType::Invalid || ((FChaosUserData*)UserData)->Type > EChaosUserDataType::CustomPayload; }
 };
 
-#if PHYSICS_INTERFACE_PHYSX
-using FUserData = FPhysxUserData;
-using FPhysicsQueryFlag = physx::PxQueryFlag;
-#else
 using FUserData = FChaosUserData;
 using FPhysicsQueryFlag = FChaosQueryFlag;
-#endif
-
 
 template <> FORCEINLINE FBodyInstance* FChaosUserData::Get(void* UserData)			{ if (!UserData || ((FChaosUserData*)UserData)->Type != EChaosUserDataType::BodyInstance) { return nullptr; } return (FBodyInstance*)((FChaosUserData*)UserData)->Payload; }
 template <> FORCEINLINE UPhysicalMaterial* FChaosUserData::Get(void* UserData)		{ if (!UserData || ((FChaosUserData*)UserData)->Type != EChaosUserDataType::PhysicalMaterial) { return nullptr; } return (UPhysicalMaterial*)((FChaosUserData*)UserData)->Payload; }
diff --git a/Engine/Source/Runtime/PhysicsCore/Public/PhysicsInterfaceUtilsCore.h b/Engine/Source/Runtime/PhysicsCore/Public/PhysicsInterfaceUtilsCore.h
--- a/Engine/Source/Runtime/PhysicsCore/Public/PhysicsInterfaceUtilsCore.h
+++ b/Engine/Source/Runtime/PhysicsCore/Public/PhysicsInterfaceUtilsCore.h
@@ -5,175 +5,10 @@
 #include "PhysicsInterfaceTypesCore.h"
 #include "PhysicsInterfaceDeclaresCore.h"
 
-#if PHYSICS_INTERFACE_PHYSX
-#include "PhysXPublicCore.h"
-#include "Containers/Union.h"
-#include "CollisionShape.h"
-#endif
-
 PHYSICSCORE_API FCollisionFilterData C2UFilterData(const FChaosFilterData& FilterData);
 PHYSICSCORE_API FChaosFilterData U2CFilterData(const FCollisionFilterData& FilterData);
 
-#if PHYSICS_INTERFACE_PHYSX
-PHYSICSCORE_API FCollisionFilterData ToUnrealFilterData(const PxFilterData& FilterData);
-#else
 PHYSICSCORE_API FCollisionFilterData ToUnrealFilterData(const FChaosFilterData& FilterData);
-#endif
-
-#if PHYSICS_INTERFACE_PHYSX
-PHYSICSCORE_API FCollisionFilterData P2UFilterData(const PxFilterData& PFilterData);
-PHYSICSCORE_API PxFilterData U2PFilterData(const FCollisionFilterData& FilterData);
-PHYSICSCORE_API PxShapeFlags BuildPhysXShapeFlags(FBodyCollisionFlags BodyCollisionFlags, bool bPhysicsStatic, bool bIsTriangleMesh);
-PHYSICSCORE_API PxGeometryType::Enum U2PCollisionShapeType(ECollisionShapeType InUType);
-PHYSICSCORE_API ECollisionShapeType P2UCollisionShapeType(PxGeometryType::Enum InPType);
-
-template<typename AGGREGATE_FLAG_TYPE, typename FLAG_TYPE>
-inline void ModifyFlag_Default(AGGREGATE_FLAG_TYPE& Flags, const FLAG_TYPE FlagToSet, const bool bValue)
-{
-	if (bValue)
-	{
-		Flags |= FlagToSet;
-	}
-	else
-	{
-		Flags.clear(FlagToSet);
-	}
-}
-
-template<const PxActorFlag::Enum FlagToSet>
-inline void ModifyActorFlag(physx::PxActorFlags& Flags, const bool bValue)
-{
-	ModifyFlag_Default(Flags, FlagToSet, bValue);
-}
-
-template<const PxShapeFlag::Enum FlagToSet>
-inline void ModifyShapeFlag(physx::PxShapeFlags& Flags, const bool bValue)
-{
-	ModifyFlag_Default(Flags, FlagToSet, bValue);
-}
-
-template<const PxRigidBodyFlag::Enum FlagToSet>
-inline void ModifyRigidBodyFlag(physx::PxRigidBodyFlags& Flags, const bool bValue)
-{
-	ModifyFlag_Default(Flags, FlagToSet, bValue);
-}
-
-template<>
-inline void ModifyRigidBodyFlag<PxRigidBodyFlag::eKINEMATIC>(physx::PxRigidBodyFlags& Flags, const bool bValue)
-{
-	// Objects can't be CCD and Kinematic at the same time.
-	// If enabling Kinematic and CCD is on, disable it and turn on Speculative CCD instead.
-	if (bValue && Flags.isSet(PxRigidBodyFlag::eENABLE_CCD))
-	{
-		Flags |= PxRigidBodyFlag::eKINEMATIC;
-		Flags |= PxRigidBodyFlag::eENABLE_SPECULATIVE_CCD;
-		Flags.clear(PxRigidBodyFlag::eENABLE_CCD);
-	}
-
-	// If disabling Kinematic and Speculative CCD is on, disable it and turn on CCD instead.
-	else if (!bValue && Flags.isSet(PxRigidBodyFlag::eENABLE_SPECULATIVE_CCD))
-	{
-		Flags |= PxRigidBodyFlag::eENABLE_CCD;
-		Flags.clear(PxRigidBodyFlag::eENABLE_SPECULATIVE_CCD);
-		Flags.clear(PxRigidBodyFlag::eKINEMATIC);
-	}
-
-	// No sanitization is needed.
-	else
-	{
-		ModifyFlag_Default(Flags, PxRigidBodyFlag::eKINEMATIC, bValue);
-	}
-}
-
-template<>
-inline void ModifyRigidBodyFlag<PxRigidBodyFlag::eENABLE_CCD>(physx::PxRigidBodyFlags& Flags, const bool bValue)
-{
-	// Objects can't be CCD and Kinematic at the same time.
-	// If disabling CCD and Speculative CCD is on, disable it too.
-	if (!bValue && Flags.isSet(PxRigidBodyFlag::eENABLE_SPECULATIVE_CCD))
-	{
-		// CCD shouldn't be enabled at this point, but force disable just in case.
-		Flags.clear(PxRigidBodyFlag::eENABLE_CCD);
-		Flags.clear(PxRigidBodyFlag::eENABLE_SPECULATIVE_CCD);
-	}
-
-	// If enabling CCD but Kinematic is on, enable Speculative CCD instead.
-	else if (bValue && Flags.isSet(PxRigidBodyFlag::eKINEMATIC))
-	{
-		Flags |= PxRigidBodyFlag::eENABLE_SPECULATIVE_CCD;
-	}
-
-	// No sanitization is needed.
-	else
-	{
-		ModifyFlag_Default(Flags, PxRigidBodyFlag::eENABLE_CCD, bValue);
-	}
-}
-
-template<const PxActorFlag::Enum FlagToSet>
-inline void ModifyActorFlag_Isolated(PxActor* PActor, const bool bValue)
-{
-	PxActorFlags ActorFlags = PActor->getActorFlags();
-	ModifyActorFlag<FlagToSet>(ActorFlags, bValue);
-	PActor->setActorFlags(ActorFlags);
-}
-
-template<const PxRigidBodyFlag::Enum FlagToSet>
-inline void ModifyRigidBodyFlag_Isolated(PxRigidBody* PRigidBody, const bool bValue)
-{
-	PxRigidBodyFlags RigidBodyFlags = PRigidBody->getRigidBodyFlags();
-	ModifyRigidBodyFlag<FlagToSet>(RigidBodyFlags, bValue);
-	PRigidBody->setRigidBodyFlags(RigidBodyFlags);
-}
-
-template<const PxShapeFlag::Enum FlagToSet>
-inline void ModifyShapeFlag_Isolated(PxShape* PShape, const bool bValue)
-{
-	PxShapeFlags ShapeFlags = PShape->getFlags();
-	ModifyShapeFlag<FlagToSet>(ShapeFlags, bValue);
-	PShape->setFlags(ShapeFlags);
-}
-
-// MISC
-
-/** Convert from unreal to physx capsule rotation */
-PHYSICSCORE_API FQuat ConvertToPhysXCapsuleRot(const FQuat& GeomRot);
-/** Convert from physx to unreal capsule rotation */
-PHYSICSCORE_API FQuat ConvertToUECapsuleRot(const PxQuat & GeomRot);
-PHYSICSCORE_API FQuat ConvertToUECapsuleRot(const FQuat & GeomRot);
-/** Convert from unreal to physx capsule pose */
-PHYSICSCORE_API PxTransform ConvertToPhysXCapsulePose(const FTransform& GeomPose);
-
-// FILTER DATA
-
-// Adapts a FCollisionShape to a PxGeometry type, used for various queries
-struct PHYSICSCORE_API FPhysXShapeAdapter
-{
-public:
-	FPhysXShapeAdapter(const FQuat& Rot, const FCollisionShape& CollisionShape);
-
-	const PxGeometry& GetGeometry() const
-	{
-		return GeometryHolder.any();
-	}
-
-public:
-	FTransform GetGeomPose(const FVector& Pos) const
-	{
-		return FTransform(Rotation, Pos);
-	}
-
-	const FQuat& GetGeomOrientation() const
-	{
-		return Rotation;
-	}
-
-private:
-	PxGeometryHolder GeometryHolder;
-	FQuat Rotation;
-};
-
-#endif // WITH_PHYX
 
 //Find the face index for a given hit. This gives us a chance to modify face index based on things like most opposing normal
 
diff --git a/Engine/Source/Runtime/PhysicsCore/Public/PhysicsPublicCore.h b/Engine/Source/Runtime/PhysicsCore/Public/PhysicsPublicCore.h
--- a/Engine/Source/Runtime/PhysicsCore/Public/PhysicsPublicCore.h
+++ b/Engine/Source/Runtime/PhysicsCore/Public/PhysicsPublicCore.h
@@ -13,47 +13,3 @@ public:
 	DECLARE_MULTICAST_DELEGATE_OneParam(FOnUpdatePhysXMaterial, UPhysicalMaterial*);
 	static FOnUpdatePhysXMaterial OnUpdatePhysXMaterial;
 };
-
-
-#if WITH_PHYSX
-namespace physx
-{
-	class PxPhysics;
-	class PxMaterial;
-}
-
-/** Pointer to PhysX SDK object */
-extern PHYSICSCORE_API physx::PxPhysics* GPhysXSDK;
-/** Pointer to PhysX allocator */
-extern PHYSICSCORE_API class FPhysXAllocator* GPhysXAllocator;
-
-extern PHYSICSCORE_API class IPhysXCookingModule* GetPhysXCookingModule(bool bForceLoad = true);
-
-/** Array of PxMaterial objects which are awaiting cleaning up. */
-extern PHYSICSCORE_API TArray<physx::PxMaterial*> GPhysXPendingKillMaterial;
-
-#if WITH_APEX
-
-namespace nvidia
-{
-	namespace apex
-	{
-		class ApexSDK;
-		class Module;
-		class ModuleClothing;
-	}
-}
-
-using namespace nvidia;
-
-/** Pointer to APEX SDK object */
-extern PHYSICSCORE_API apex::ApexSDK* GApexSDK;
-/** Pointer to APEX legacy module object */
-extern PHYSICSCORE_API apex::Module* GApexModuleLegacy;
-#if WITH_APEX_CLOTHING
-/** Pointer to APEX Clothing module object */
-extern PHYSICSCORE_API apex::ModuleClothing* GApexModuleClothing;
-#endif //WITH_APEX_CLOTHING
-
-#endif // #if WITH_APEX
-#endif
\ No newline at end of file
diff --git a/Engine/Source/Runtime/PhysicsCore/Public/SQAccelerator.h b/Engine/Source/Runtime/PhysicsCore/Public/SQAccelerator.h
--- a/Engine/Source/Runtime/PhysicsCore/Public/SQAccelerator.h
+++ b/Engine/Source/Runtime/PhysicsCore/Public/SQAccelerator.h
@@ -4,21 +4,9 @@
 
 #include "Math/BoxSphereBounds.h"
 
-#if PHYSICS_INTERFACE_PHYSX
-#include "PhysXInterfaceWrapperCore.h"	//todo remove physx specific include once refactor is done
-#elif WITH_CHAOS
 #include "ChaosInterfaceWrapperCore.h"
-#endif
-
 #include "ChaosSQTypes.h"
 
-#if WITH_PHYSX
-namespace physx
-{
-	class PxScene;
-}
-#endif
-
 namespace Chaos
 {
 	class FImplicitObject;
@@ -60,7 +48,6 @@ private:
 };
 
 // An interface to the scene query accelerator that allows us to run queries against either PhysX or Chaos
-// when compiling WITH_PHYSX.
 // This was used in the 2019 GDC demos and is now broken. To make it work again, we would need to implement
 // the FChaosSQAcceleratorAdapter below to use its internal SQ accelerator and convert the inputs and outputs
 // from/to PhysX types.
@@ -87,44 +74,3 @@ public:
 private:
 	TArray<ISQAccelerator*> Accelerators;
 };
-
-#if WITH_PHYSX
-// A Chaos Query Accelerator with a PhysX API
-// TODO: Not implemented - required to make GDC 2019 demos work again.
-class PHYSICSCORE_API FChaosSQAcceleratorAdapter : public ISQAccelerator
-{
-public:
-
-	FChaosSQAcceleratorAdapter(const Chaos::ISpatialAcceleration<Chaos::FAccelerationStructureHandle, Chaos::FReal, 3>& InSpatialAcceleration);
-	virtual ~FChaosSQAcceleratorAdapter() {};
-
-	virtual void Raycast(const FVector& Start, const FVector& Dir, const float DeltaMagnitude, FPhysicsHitCallback<FHitRaycast>& HitBuffer, EHitFlags OutputFlags, const FQueryFilterData& QueryFilterData, ICollisionQueryFilterCallbackBase& QueryCallback) const override;
-	virtual void Sweep(const FPhysicsGeometry& QueryGeom, const FTransform& StartTM, const FVector& Dir, const float DeltaMagnitude, FPhysicsHitCallback<FHitSweep>& HitBuffer, EHitFlags OutputFlags, const FQueryFilterData& QueryFilterData, ICollisionQueryFilterCallbackBase& QueryCallback) const override;
-	virtual void Overlap(const FPhysicsGeometry& QueryGeom, const FTransform& GeomPose, FPhysicsHitCallback<FHitOverlap>& HitBuffer, const FQueryFilterData& QueryFilterData, ICollisionQueryFilterCallbackBase& QueryCallback) const override;
-
-private:
-	FChaosSQAccelerator ChaosSQAccelerator;
-};
-#endif
-
-
-#if WITH_PHYSX && !WITH_CHAOS
-class PHYSICSCORE_API FPhysXSQAccelerator : public ISQAccelerator
-{
-public:
-
-	FPhysXSQAccelerator();
-	FPhysXSQAccelerator(physx::PxScene* InScene);
-	virtual ~FPhysXSQAccelerator() {};
-
-	virtual void Raycast(const FVector& Start, const FVector& Dir, const float DeltaMagnitude, FPhysicsHitCallback<FHitRaycast>& HitBuffer, EHitFlags OutputFlags, const FQueryFilterData& QueryFilterData, ICollisionQueryFilterCallbackBase& QueryCallback) const override;
-	virtual void Sweep(const FPhysicsGeometry& QueryGeom, const FTransform& StartTM, const FVector& Dir, const float DeltaMagnitude, FPhysicsHitCallback<FHitSweep>& HitBuffer, EHitFlags OutputFlags, const FQueryFilterData& QueryFilterData, ICollisionQueryFilterCallbackBase& QueryCallback) const override;
-	virtual void Overlap(const FPhysicsGeometry& QueryGeom, const FTransform& GeomPose, FPhysicsHitCallback<FHitOverlap>& HitBuffer, const FQueryFilterData& QueryFilterData, ICollisionQueryFilterCallbackBase& QueryCallback) const override;
-
-	void SetScene(physx::PxScene* InScene);
-
-private:
-	physx::PxScene* Scene;
-
-};
-#endif
diff --git a/Engine/Source/Runtime/PhysicsCore/Public/SQCapture.h b/Engine/Source/Runtime/PhysicsCore/Public/SQCapture.h
--- a/Engine/Source/Runtime/PhysicsCore/Public/SQCapture.h
+++ b/Engine/Source/Runtime/PhysicsCore/Public/SQCapture.h
@@ -41,17 +41,6 @@ struct PHYSICSCORE_API FSQCapture
 		Overlap
 	} SQType;
 
-#if PHYSICS_INTERFACE_PHYSX
-	void StartCapturePhysXSweep(const PxScene& Scene, const PxGeometry& InQueryGeom, const FTransform& InStartTM, const FVector& InDir, float InDeltaMag, FHitFlags InOutputFlags, const FQueryFilterData& QueryFilter, const FCollisionFilterData& FilterData, ICollisionQueryFilterCallbackBase& Callback);
-	void EndCapturePhysXSweep(const PxHitCallback<PxSweepHit>& Results);
-
-	void StartCapturePhysXRaycast(const PxScene& Scene, const FVector& InStartPoint, const FVector& InDir, float InDeltaMag, FHitFlags InOutputFlags, const FQueryFilterData& QueryFilter, const FCollisionFilterData& FilterData, ICollisionQueryFilterCallbackBase& Callback);
-	void EndCapturePhysXRaycast(const PxHitCallback<PxRaycastHit>& Results);
-
-	void StartCapturePhysXOverlap(const PxScene& Scene, const PxGeometry& InQueryGeom, const FTransform& WorldTM, const FQueryFilterData& QueryFilter, const FCollisionFilterData& FilterData, ICollisionQueryFilterCallbackBase& Callback);
-	void EndCapturePhysXOverlap(const PxHitCallback<PxOverlapHit>& Results);
-#endif
-
 	void StartCaptureChaosSweep(const Chaos::FPBDRigidsEvolution& Evolution, const Chaos::FImplicitObject& InQueryGeom, const FTransform& InStartTM, const FVector& InDir, float InDeltaMag, FHitFlags InOutputFlags, const FQueryFilterData& QueryFilter, const FCollisionFilterData& FilterData, ICollisionQueryFilterCallbackBase& Callback);
 	void EndCaptureChaosSweep(const ChaosInterface::FSQHitBuffer<ChaosInterface::FSweepHit>& Results);
 
@@ -62,10 +51,6 @@ struct PHYSICSCORE_API FSQCapture
 	void EndCaptureChaosOverlap(const ChaosInterface::FSQHitBuffer<ChaosInterface::FOverlapHit>& Results);
 
 	ECollisionQueryHitType GetFilterResult(const Chaos::FPerShapeData* Shape, const Chaos::FGeometryParticle* Actor) const;
-
-#if PHYSICS_INTERFACE_PHYSX
-	ECollisionQueryHitType GetFilterResult(const physx::PxShape* Shape, const physx::PxActor* Actor) const;
-#endif
 	
 	FVector Dir;
 	FTransform StartTM;	//only valid if overlap or sweep
@@ -76,18 +61,10 @@ struct PHYSICSCORE_API FSQCapture
 	FQueryFilterData QueryFilterData;
 	TUniquePtr<ICollisionQueryFilterCallbackBase> FilterCallback;
 
-#if PHYSICS_INTERFACE_PHYSX
-	PhysXInterface::FDynamicHitBuffer<PxSweepHit> PhysXSweepBuffer;
-	PhysXInterface::FDynamicHitBuffer<PxRaycastHit> PhysXRaycastBuffer;
-	PhysXInterface::FDynamicHitBuffer<PxOverlapHit> PhysXOverlapBuffer;
-	PxGeometryHolder PhysXGeometry;
-#endif
-
 	TUniquePtr<Chaos::FImplicitObject> ChaosOwnerObject;	//should be private, do not access directly
 	const Chaos::FImplicitObject* ChaosGeometry;
 	TUniquePtr<Chaos::FImplicitObject> SerializableChaosGeometry;
-#if WITH_CHAOS
-	//for now just use physx hit buffer
+
 	ChaosInterface::FSQHitBuffer<ChaosInterface::FSweepHit> ChaosSweepBuffer;
 	TArray<ChaosInterface::FSweepHit> ChaosSweepTouches;
 
@@ -96,65 +73,16 @@ struct PHYSICSCORE_API FSQCapture
 
 	ChaosInterface::FSQHitBuffer<ChaosInterface::FOverlapHit> ChaosOverlapBuffer;
 	TArray<ChaosInterface::FOverlapHit> ChaosOverlapTouches;
-#endif
 
 private:
 	FSQCapture(FPhysTestSerializer& OwningPhysSerializer);	//This should be created by PhysTestSerializer
 	void Serialize(Chaos::FChaosArchive& Ar);
 
-#if PHYSICS_INTERFACE_PHYSX
-	void SerializeActorToShapeHitsArray(FArchive& Ar);
-#endif
-
 	friend FPhysTestSerializer;
 
 	TArray<uint8> GeomData;
 	TArray<uint8> HitData;
 
-#if 0
-	void CreateChaosDataFromPhysX();
-	void CreateChaosFilterResults();
-#endif
-
-#if PHYSICS_INTERFACE_PHYSX
-
-	static constexpr uint64 ShapeCollectionID = 1;
-
-	void CapturePhysXFilterResults(const PxScene& Scene, const FCollisionFilterData& FilterData, ICollisionQueryFilterCallbackBase& Callback);
-	void CreatePhysXData();
-	void SetPhysXGeometryData(const PxGeometry& Geometry);
-
-	physx::PxActor* GetTransientActor(physx::PxActor* Actor) const;
-
-	physx::PxShape* GetTransientShape(physx::PxShape* Shape) const;
-
-	void SerializePhysXHitType(FArchive& Ar, PxOverlapHit& Hit);
-
-	template <typename T>
-	void SerializePhysXHitType(FArchive& Ar, T& Hit);
-
-	template <typename THit>
-	void SerializePhysXBuffers(FArchive& Ar, int32 Version, PhysXInterface::FDynamicHitBuffer<THit>& PhysXBuffer);
-
-	struct FPhysXSerializerData
-	{
-		FPhysXSerializerData(int32 NumBytes);
-		~FPhysXSerializerData();
-
-		void* Data;
-		physx::PxShape* Shape;	//holder for geometry so we can serialize it out in a collection
-		physx::PxCollection* Collection;
-		physx::PxSerializationRegistry* Registry;
-	};
-	TUniquePtr<FPhysXSerializerData> AlignedDataHelper;
-
-	TMap<physx::PxActor*, TArray<TPair<physx::PxShape*, ECollisionQueryHitType>>> PxActorToShapeHitsArray;
-
-	//only valid during capture when serializing out runtime structures that already use non transient data
-	TMap<physx::PxActor*, physx::PxActor*> NonTransientToTransientActors;
-	TMap<physx::PxShape*, physx::PxShape*> NonTransientToTransientShapes;
-#endif
-
 	FPhysTestSerializer& PhysSerializer;
 
 	void CaptureChaosFilterResults(const Chaos::FPBDRigidsEvolution& Evolution, const FCollisionFilterData& FilterData, ICollisionQueryFilterCallbackBase& Callback);
