SetGlobalPose_AssumesLocked doesn't do anything if both X and R haven't changed.

diff --git a/Engine/Source/Runtime/PhysicsCore/Private/ChaosEngineInterface.cpp b/Engine/Source/Runtime/PhysicsCore/Private/ChaosEngineInterface.cpp
--- a/Engine/Source/Runtime/PhysicsCore/Private/ChaosEngineInterface.cpp
+++ b/Engine/Source/Runtime/PhysicsCore/Private/ChaosEngineInterface.cpp
@@ -1780,6 +1780,11 @@ FChaosScene* FChaosEngineInterface::GetCurrentScene(const FPhysicsActorHandle& I
 void FChaosEngineInterface::SetGlobalPose_AssumesLocked(const FPhysicsActorHandle& InActorReference,const FTransform& InNewPose,bool bAutoWake)
 {
 	Chaos::FRigidBodyHandle_External& Body_External = InActorReference->GetGameThreadAPI();
+	if (Chaos::FVec3::IsNearlyEqual(InNewPose.GetLocation(), Body_External.X(), SMALL_NUMBER) && Chaos::FRotation3::IsNearlyEqual(InNewPose.GetRotation(), Body_External.R(), SMALL_NUMBER))
+	{
+		// don't update X/R if they haven't changed. this allows scale to be set on simulating body without overriding async position/rotation.
+		return;
+	}
 	Body_External.SetX(InNewPose.GetLocation());
 	Body_External.SetR(InNewPose.GetRotation());
 	Body_External.UpdateShapeBounds();
