Chaos:
SetGlobalPose now no longer discards the same XR if you're kinematic or sleeping.

diff --git a/Engine/Source/Runtime/PhysicsCore/Private/ChaosEngineInterface.cpp b/Engine/Source/Runtime/PhysicsCore/Private/ChaosEngineInterface.cpp
--- a/Engine/Source/Runtime/PhysicsCore/Private/ChaosEngineInterface.cpp
+++ b/Engine/Source/Runtime/PhysicsCore/Private/ChaosEngineInterface.cpp
@@ -1780,9 +1780,9 @@ FChaosScene* FChaosEngineInterface::GetCurrentScene(const FPhysicsActorHandle& I
 void FChaosEngineInterface::SetGlobalPose_AssumesLocked(const FPhysicsActorHandle& InActorReference,const FTransform& InNewPose,bool bAutoWake)
 {
 	Chaos::FRigidBodyHandle_External& Body_External = InActorReference->GetGameThreadAPI();
-	if (Chaos::FVec3::IsNearlyEqual(InNewPose.GetLocation(), Body_External.X(), SMALL_NUMBER) && Chaos::FRotation3::IsNearlyEqual(InNewPose.GetRotation(), Body_External.R(), SMALL_NUMBER))
+	if (!IsKinematic(InActorReference) && !IsSleeping(InActorReference) && Chaos::FVec3::IsNearlyEqual(InNewPose.GetLocation(), Body_External.X(), SMALL_NUMBER) && Chaos::FRotation3::IsNearlyEqual(InNewPose.GetRotation(), Body_External.R(), SMALL_NUMBER))
 	{
-		// don't update X/R if they haven't changed. this allows scale to be set on simulating body without overriding async position/rotation.
+		// if simulating, don't update X/R if they haven't changed. this allows scale to be set on simulating body without overriding async position/rotation.
 		return;
 	}
 	Body_External.SetX(InNewPose.GetLocation());
