Fixed crash attempting to retrieve shapes from a body instance that has no actor instance (e.g. welded bodies).
Fixed incorrect identification of owned instance shapes in GetSquaredDistanceToBody for welded instances

diff --git a/Engine/Source/Runtime/Engine/Private/PhysicsEngine/BodyInstance.cpp b/Engine/Source/Runtime/Engine/Private/PhysicsEngine/BodyInstance.cpp
--- a/Engine/Source/Runtime/Engine/Private/PhysicsEngine/BodyInstance.cpp
+++ b/Engine/Source/Runtime/Engine/Private/PhysicsEngine/BodyInstance.cpp
@@ -430,11 +430,16 @@ bool FBodyInstance::IsShapeBoundToBody(const FPhysicsShapeHandle& Shape) const
 const TMap<FPhysicsShapeHandle, FBodyInstance::FWeldInfo>* FBodyInstance::GetCurrentWeldInfo() const
 {
 	return ShapeToBodiesMap.Get();
-	}
+}
 
 int32 FBodyInstance::GetAllShapes_AssumesLocked(TArray<FPhysicsShapeHandle>& OutShapes) const
+{
+	if(ActorHandle)
 	{
-	return FPhysicsInterface::GetAllShapes_AssumedLocked(ActorHandle, OutShapes);
+		return FPhysicsInterface::GetAllShapes_AssumedLocked(ActorHandle, OutShapes);
+	}
+
+	return 0;
 }
 
 void FBodyInstance::UpdateTriMeshVertices(const TArray<FVector> & NewPositions)
diff --git a/Engine/Source/Runtime/Engine/Private/PhysicsEngine/Experimental/PhysInterface_Chaos.cpp b/Engine/Source/Runtime/Engine/Private/PhysicsEngine/Experimental/PhysInterface_Chaos.cpp
--- a/Engine/Source/Runtime/Engine/Private/PhysicsEngine/Experimental/PhysInterface_Chaos.cpp
+++ b/Engine/Source/Runtime/Engine/Private/PhysicsEngine/Experimental/PhysInterface_Chaos.cpp
@@ -1249,11 +1249,10 @@ bool FPhysInterface_Chaos::GetSquaredDistanceToBody(const FBodyInstance* InInsta
 
 	FPhysicsCommand::ExecuteRead(UseBI->ActorHandle, [&](const FPhysicsActorHandle& Actor)
 	{
-
 		bEarlyOut = false;
 
 		TArray<FPhysicsShapeReference_Chaos> Shapes;
-		InInstance->GetAllShapes_AssumesLocked(Shapes);
+		UseBI->GetAllShapes_AssumesLocked(Shapes);
 		for (const FPhysicsShapeReference_Chaos& Shape : Shapes)
 		{
 			if (!Shape.IsValid())
@@ -1261,7 +1260,7 @@ bool FPhysInterface_Chaos::GetSquaredDistanceToBody(const FBodyInstance* InInsta
 				continue;
 			}
 
-			if (UseBI->IsShapeBoundToBody(Shape) == false)	//skip welded shapes that do not belong to us
+			if (InInstance->IsShapeBoundToBody(Shape) == false)	//skip welded shapes that do not belong to us
 			{
 				continue;
 			}
diff --git a/Engine/Source/Runtime/PhysicsCore/Private/ChaosEngineInterface.cpp b/Engine/Source/Runtime/PhysicsCore/Private/ChaosEngineInterface.cpp
--- a/Engine/Source/Runtime/PhysicsCore/Private/ChaosEngineInterface.cpp
+++ b/Engine/Source/Runtime/PhysicsCore/Private/ChaosEngineInterface.cpp
@@ -1678,14 +1678,20 @@ void FChaosEngineInterface::SetLocalTransform(const FPhysicsShapeHandle& InShape
 template<typename AllocatorType>
 int32 GetAllShapesInternalImp_AssumedLocked(const FPhysicsActorHandle& InActorHandle,TArray<FPhysicsShapeReference_Chaos,AllocatorType>& OutShapes)
 {
-	const Chaos::FShapesArray& ShapesArray = InActorHandle->GetGameThreadAPI().ShapesArray();
-	OutShapes.Reset(ShapesArray.Num());
-	//todo: can we avoid this construction?
-	for(const TUniquePtr<Chaos::FPerShapeData>& Shape : ShapesArray)
+	if(InActorHandle)
 	{
-		OutShapes.Add(FPhysicsShapeReference_Chaos(Shape.Get(),InActorHandle));
+		const Chaos::FShapesArray& ShapesArray = InActorHandle->GetGameThreadAPI().ShapesArray();
+		OutShapes.Reset(ShapesArray.Num());
+		//todo: can we avoid this construction?
+		for(const TUniquePtr<Chaos::FPerShapeData>& Shape : ShapesArray)
+		{
+			OutShapes.Add(FPhysicsShapeReference_Chaos(Shape.Get(),InActorHandle));
+		}
+
+		return OutShapes.Num();
 	}
-	return OutShapes.Num();
+
+	return 0;
 }
 
 int32 FChaosEngineInterface::GetAllShapes_AssumedLocked(const FPhysicsActorHandle& InActorHandle,TArray<FPhysicsShapeReference_Chaos,FDefaultAllocator>& OutShapes)
