Chaos: Remove ensure in GetWorldVelocityAtPoint that required the object to be kinematic or dynamic. It is valid to request a velocity from a static object without firing an ensure - the velocity will always be zero in that case.
In cases where a user is calculating velocity delta between bodies it's valid to just compare two calls to GetWorldVelocityAtPoint rather than having to handle the if is static case at all callsites. This is also how the interface functioned previously.

diff --git a/Engine/Source/Runtime/PhysicsCore/Private/ChaosEngineInterface.cpp b/Engine/Source/Runtime/PhysicsCore/Private/ChaosEngineInterface.cpp
--- a/Engine/Source/Runtime/PhysicsCore/Private/ChaosEngineInterface.cpp
+++ b/Engine/Source/Runtime/PhysicsCore/Private/ChaosEngineInterface.cpp
@@ -659,7 +659,7 @@ FVector FChaosEngineInterface::GetWorldVelocityAtPoint_AssumesLocked(const FPhys
 	if(ensure(FChaosEngineInterface::IsValid(InActorReference)))
 	{
 		const Chaos::FRigidBodyHandle_External& Body_External = InActorReference->GetGameThreadAPI();
-		if(ensure(Body_External.CanTreatAsKinematic()))
+		if(Body_External.CanTreatAsKinematic())
 		{
 			const bool bIsRigid = Body_External.CanTreatAsRigid();
 			const Chaos::FVec3 COM = bIsRigid ? Chaos::FParticleUtilitiesGT::GetCoMWorldPosition(&Body_External) : (Chaos::FVec3)Chaos::FParticleUtilitiesGT::GetActorWorldTransform(&Body_External).GetTranslation();
