[PhysicsCore]
* Changed "enum EPhysicalSurface" to "enum EPhysicalSurface : int" to be able to forward declare it (it ends up being the same thing, just explicit)

diff --git a/Engine/Source/Runtime/PhysicsCore/Public/Chaos/ChaosEngineInterface.h b/Engine/Source/Runtime/PhysicsCore/Public/Chaos/ChaosEngineInterface.h
--- a/Engine/Source/Runtime/PhysicsCore/Public/Chaos/ChaosEngineInterface.h
+++ b/Engine/Source/Runtime/PhysicsCore/Public/Chaos/ChaosEngineInterface.h
@@ -15,7 +15,7 @@
 
 /** Types of surfaces in the game, used by Physical Materials */
 UENUM(BlueprintType)
-enum EPhysicalSurface
+enum EPhysicalSurface : int
 {
 	SurfaceType_Default UMETA(DisplayName="Default"),
 	SurfaceType1 UMETA(Hidden),
