Chaos: Protected call to IsStatic in the interface layer matching previous behaviour (protects against crashes when setting up body instance parameters in component constructors)

diff --git a/Engine/Source/Runtime/PhysicsCore/Private/ChaosEngineInterface.cpp b/Engine/Source/Runtime/PhysicsCore/Private/ChaosEngineInterface.cpp
--- a/Engine/Source/Runtime/PhysicsCore/Private/ChaosEngineInterface.cpp
+++ b/Engine/Source/Runtime/PhysicsCore/Private/ChaosEngineInterface.cpp
@@ -380,13 +380,17 @@ bool FChaosEngineInterface::IsRigidBody(const FPhysicsActorHandle& InActorRefere
 
 bool FChaosEngineInterface::IsDynamic(const FPhysicsActorHandle& InActorReference)
 {
-	// Do this to match the PhysX interface behavior: :( :( :(
 	return !IsStatic(InActorReference);
 }
 
 bool FChaosEngineInterface::IsStatic(const FPhysicsActorHandle& InActorReference)
 {
-	return InActorReference->GetGameThreadAPI().ObjectState() == Chaos::EObjectStateType::Static;
+	if(FChaosEngineInterface::IsValid(InActorReference))
+	{
+		return InActorReference->GetGameThreadAPI().ObjectState() == Chaos::EObjectStateType::Static;
+	}
+
+	return false;
 }
 
 bool FChaosEngineInterface::IsKinematic(const FPhysicsActorHandle& InActorReference)
