Chaos - fix kinematic teleport if it happens after animation by clearing the kinematic target in SetGlobalPose

diff --git a/Engine/Source/Runtime/PhysicsCore/Private/ChaosEngineInterface.cpp b/Engine/Source/Runtime/PhysicsCore/Private/ChaosEngineInterface.cpp
--- a/Engine/Source/Runtime/PhysicsCore/Private/ChaosEngineInterface.cpp
+++ b/Engine/Source/Runtime/PhysicsCore/Private/ChaosEngineInterface.cpp
@@ -1819,8 +1819,25 @@ void FChaosEngineInterface::SetGlobalPose_AssumesLocked(const FPhysicsActorHandl
 		// if simulating, don't update X/R if they haven't changed. this allows scale to be set on simulating body without overriding async position/rotation.
 		return;
 	}
-	Body_External.SetX(InNewPose.GetLocation());
-	Body_External.SetR(InNewPose.GetRotation());
+
+	if (IsKinematic(InActorReference))
+	{
+		// If we have already set a kinematic target this tick it must be cleared. NOTE: because we are clearing the target, 
+		// the Velocity will be zeroed on the next simulation tick unless there is a subsequent call to SetKinematicTarget 
+		// with an appropriate position delta. 
+		Body_External.ClearKinematicTarget();
+		Body_External.SetX(InNewPose.GetLocation());
+		Body_External.SetR(InNewPose.GetRotation());
+		// @todo(chaos): Should we also explicitly set the velocity to zero here?
+		//Body_External.SetV(FVector::Zero());
+		//Body_External.SetW(FVector::Zero());
+	}
+	else
+	{
+		Body_External.SetX(InNewPose.GetLocation());
+		Body_External.SetR(InNewPose.GetRotation());
+	}
+
 	Body_External.UpdateShapeBounds();
 
 	FChaosScene* Scene = GetCurrentScene(InActorReference);
@@ -1829,9 +1846,8 @@ void FChaosEngineInterface::SetGlobalPose_AssumesLocked(const FPhysicsActorHandl
 
 void FChaosEngineInterface::SetKinematicTarget_AssumesLocked(const FPhysicsActorHandle& InActorReference,const FTransform& InNewTarget)
 {
-	Chaos::TKinematicTarget<Chaos::FReal, 3> NewKinematicTarget;
 	// SetKinematicTarget_AssumesLocked could be called multiple times in one time step
-	NewKinematicTarget.SetTargetMode(InNewTarget);
+	const Chaos::FKinematicTarget NewKinematicTarget = Chaos::FKinematicTarget::MakePositionTarget(InNewTarget);
 	InActorReference->GetGameThreadAPI().SetKinematicTarget(NewKinematicTarget);
 
 	// IMPORTANT : we do not invalidate X and R as they will be properly computed using the kinematic target information 
