Chaos Physics : add material strength to the physical material
- also deprecate the DestructibleDamageThresholdScale porperty as it is not used anymore

diff --git a/Engine/Source/Runtime/PhysicsCore/Private/PhysicalMaterial.cpp b/Engine/Source/Runtime/PhysicsCore/Private/PhysicalMaterial.cpp
--- a/Engine/Source/Runtime/PhysicsCore/Private/PhysicalMaterial.cpp
+++ b/Engine/Source/Runtime/PhysicsCore/Private/PhysicalMaterial.cpp
@@ -27,8 +27,12 @@ UPhysicalMaterial::UPhysicalMaterial(const FObjectInitializer& ObjectInitializer
 	SleepLinearVelocityThreshold = 1.f;
 	SleepAngularVelocityThreshold = 0.05f;
 	SleepCounterThreshold = 4;
-	DestructibleDamageThresholdScale = 1.0f;
 	bOverrideFrictionCombineMode = false;
+	// using concrete as default ( lowest values of it )
+	TensileStrength = 2;
+	CompressionStrength = 20;
+	ShearStrength = 6;
+
 	UserData = FChaosUserData(this);
 }
 
diff --git a/Engine/Source/Runtime/PhysicsCore/Public/PhysicalMaterials/PhysicalMaterial.h b/Engine/Source/Runtime/PhysicsCore/Public/PhysicalMaterials/PhysicalMaterial.h
--- a/Engine/Source/Runtime/PhysicsCore/Public/PhysicalMaterials/PhysicalMaterial.h
+++ b/Engine/Source/Runtime/PhysicsCore/Public/PhysicalMaterials/PhysicalMaterial.h
@@ -90,8 +90,9 @@ class PHYSICSCORE_API UPhysicalMaterial : public UObject
 	float RaiseMassToPower;
 
 	/** How much to scale the damage threshold by on any destructible we are applied to */
-	UPROPERTY(EditAnywhere, BlueprintReadOnly, Category = Destruction)
-	float DestructibleDamageThresholdScale;
+	UE_DEPRECATED(5.3, "This property is not used anywhere, use Geometry Collection damage threshold related features instead")
+	UPROPERTY()
+	float DestructibleDamageThresholdScale_DEPRECATED;
 
 	UPROPERTY()
 	TObjectPtr<class UDEPRECATED_PhysicalMaterialPropertyBase> PhysicalMaterialProperty_DEPRECATED;
@@ -102,6 +103,25 @@ class PHYSICSCORE_API UPhysicalMaterial : public UObject
 	UPROPERTY(EditAnywhere, BlueprintReadOnly, Category = PhysicalProperties)
 	TEnumAsByte<EPhysicalSurface> SurfaceType;
 
+	/** 
+	* Tensile strength of the material in MegaPascal ( 10^6 N/m2 )
+	* This amount of tension force per area the material can withstand before it fractures
+	*/
+	float TensileStrength;
+
+	/**
+	* Compression strength of the material in MegaPascal ( 10^6 N/m2 )
+	* This amount of compression force per area the material can withstand before it fractures, crumbles or buckles
+	*/
+	float CompressionStrength;
+
+	/**
+	* Shear strength of the material in MegaPascal ( 10^6 N/m2 )
+	* This amount of shear force per area the material can withstand before it fractures
+	*/
+	float ShearStrength;
+
+
 public:
 
 	TUniquePtr<FPhysicsMaterialHandle> MaterialHandle;
