utilize the FaceNormal in the sweep even when we hit implicit objects without a face index

diff --git a/Engine/Source/Runtime/Engine/Private/Collision/CollisionConversions.cpp b/Engine/Source/Runtime/Engine/Private/Collision/CollisionConversions.cpp
--- a/Engine/Source/Runtime/Engine/Private/Collision/CollisionConversions.cpp
+++ b/Engine/Source/Runtime/Engine/Private/Collision/CollisionConversions.cpp
@@ -100,8 +100,11 @@ static FVector FindGeomOpposingNormal(ECollisionShapeType QueryGeomType, const T
 	// TODO: can we support other shapes here as well?
 	if (QueryGeomType == ECollisionShapeType::Capsule || QueryGeomType == ECollisionShapeType::Sphere)
 	{
-		const FPhysicsShape* Shape = GetShape(Hit);
-		if (Shape)
+		if (!Hit.FaceNormal.IsNearlyZero())
+		{
+			return Hit.FaceNormal;
+		}
+		else if (const FPhysicsShape* Shape = GetShape(Hit))
 		{
 			const FTransform ActorTM(Hit.Actor->R(), Hit.Actor->X());
 			const FVector LocalInNormal = ActorTM.InverseTransformVectorNoScale(InNormal);
diff --git a/Engine/Source/Runtime/Experimental/Chaos/Public/Chaos/GeometryQueries.h b/Engine/Source/Runtime/Experimental/Chaos/Public/Chaos/GeometryQueries.h
--- a/Engine/Source/Runtime/Experimental/Chaos/Public/Chaos/GeometryQueries.h
+++ b/Engine/Source/Runtime/Experimental/Chaos/Public/Chaos/GeometryQueries.h
@@ -304,6 +304,7 @@ namespace Chaos
 		{
 			OutNormal = ATM.TransformVectorNoScale(LocalNormal);
 			OutPosition = ATM.TransformPositionNoScale(LocalPosition);
+			OutFaceNormal = ATM.TransformVectorNoScale(A.FindGeometryOpposingNormal(Dir, OutFaceIndex, OutNormal));
 		}
 
 		return bResult;
diff --git a/Engine/Source/Runtime/Experimental/Chaos/Public/Chaos/Interface/SQTypes.h b/Engine/Source/Runtime/Experimental/Chaos/Public/Chaos/Interface/SQTypes.h
--- a/Engine/Source/Runtime/Experimental/Chaos/Public/Chaos/Interface/SQTypes.h
+++ b/Engine/Source/Runtime/Experimental/Chaos/Public/Chaos/Interface/SQTypes.h
@@ -46,6 +46,7 @@ namespace ChaosInterface
 		*/
 
 		int32 FaceIndex; // Signed int to match TArray's size type, and so INDEX_NONE/-1 doesn't underflow.
+		FVector FaceNormal{ EForceInit::ForceInitToZero };
 
 		int32 ElementIndex; // Currently used to indicate which shape was hit for a particle with multiple shapes.
 
@@ -124,6 +125,7 @@ namespace ChaosInterface
 		*/
 
 		int32 FaceIndex = -1; // Signed int to match TArray's size type, and so INDEX_NONE/-1 doesn't underflow.
+		FVector FaceNormal{EForceInit::ForceInitToZero};
 
 		int32 ElementIndex = -1; // Currently used to indicate which shape was hit for a particle with multiple shapes.
 	};
diff --git a/Engine/Source/Runtime/PhysicsCore/Private/SQAccelerator.cpp b/Engine/Source/Runtime/PhysicsCore/Private/SQAccelerator.cpp
--- a/Engine/Source/Runtime/PhysicsCore/Private/SQAccelerator.cpp
+++ b/Engine/Source/Runtime/PhysicsCore/Private/SQAccelerator.cpp
@@ -337,15 +337,14 @@ private:
 					if constexpr(SQ == ESQType::Sweep && std::is_same_v<THitType, ChaosInterface::FSweepHit>)
 					{
 						const THitType* CurrentHit = HitBuffer.GetCurrentHit();
-						if(FaceIdx != INDEX_NONE)
+						if(FaceIdx != INDEX_NONE || !FaceNormal.IsNearlyZero())
 						{
 							if(CurrentHit)
 							{
 								constexpr static FReal CoLocationEpsilon = 1e-6;
 								const FReal DistDelta = FMath::Abs(Distance - CurrentHit->Distance);
-								const int32 OldFace = CurrentHit->FaceIndex;
 
-								if(DistDelta < CoLocationEpsilon && OldFace != INDEX_NONE)
+								if(DistDelta < CoLocationEpsilon && !HitFaceNormal.IsNearlyZero())
 								{
 									// We already have a face hit from another triangle mesh - see if this one is better (more opposing the sweep)
 									const FReal OldDot = FVec3::DotProduct(CurData->Dir, HitFaceNormal);
@@ -376,6 +375,11 @@ private:
 					{
 						FillHitHelper(Hit, Distance, WorldPosition, WorldNormal, FaceIdx, bComputeMTD);
 
+						if constexpr (std::is_base_of_v<ChaosInterface::FQueryHit, THitType> || std::is_base_of_v<ChaosInterface::FPTQueryHit, THitType>)
+						{
+							Hit.FaceNormal = HitFaceNormal;
+						}
+
 						HitType = QueryFilterData.flags & FChaosQueryFlag::ePOSTFILTER ? QueryCallback.PostFilter(QueryFilterDataConcrete, Hit) : HitType;
 
 						if(HitType != ECollisionQueryHitType::None)
