Fixed mismatched copy constructor/assignment operator in FPhysicsShapeReference_Chaos.

diff --git a/Engine/Source/Runtime/PhysicsCore/Public/Chaos/ChaosEngineInterface.h b/Engine/Source/Runtime/PhysicsCore/Public/Chaos/ChaosEngineInterface.h
--- a/Engine/Source/Runtime/PhysicsCore/Public/Chaos/ChaosEngineInterface.h
+++ b/Engine/Source/Runtime/PhysicsCore/Public/Chaos/ChaosEngineInterface.h
@@ -260,9 +260,6 @@ public:
 		: Shape(nullptr), ActorRef() { }
 	FPhysicsShapeReference_Chaos(Chaos::FPerShapeData* ShapeIn, const FPhysicsActorHandle& ActorRefIn)
 		: Shape(ShapeIn), ActorRef(ActorRefIn) { }
-	FPhysicsShapeReference_Chaos(const FPhysicsShapeReference_Chaos& Other)
-		: Shape(Other.Shape)
-		, ActorRef(Other.ActorRef){}
 
 	bool IsValid() const { return (Shape != nullptr); }
 	bool Equals(const FPhysicsShapeReference_Chaos& Other) const { return Shape == Other.Shape; }
