introduce new external context safe payload that can be used for chaos acceleration structures - revamp SQ interfaces a bit to be able to accept additional types of SQ overrides

- cluster union acceleration structure is now safe even when a child physics state is destroyed

diff --git a/Engine/Source/Runtime/Engine/Classes/PhysicsEngine/ClusterUnionComponent.h b/Engine/Source/Runtime/Engine/Classes/PhysicsEngine/ClusterUnionComponent.h
--- a/Engine/Source/Runtime/Engine/Classes/PhysicsEngine/ClusterUnionComponent.h
+++ b/Engine/Source/Runtime/Engine/Classes/PhysicsEngine/ClusterUnionComponent.h
@@ -9,6 +9,7 @@
 #include "Containers/Set.h"
 #include "Engine/EngineTypes.h"
 #include "Logging/LogMacros.h"
+#include "PhysicsEngine/ExternalSpatialAccelerationPayload.h"
 #include "PhysicsInterfaceTypesCore.h"
 #include "PhysicsProxy/ClusterUnionPhysicsProxy.h"
 #include "UObject/ObjectKey.h"
@@ -16,6 +17,7 @@
 #include "ClusterUnionComponent.generated.h"
 
 class AActor;
+struct FExternalSpatialAccelerationPayload;
 class FPhysScene_Chaos;
 class UClusterUnionReplicatedProxyComponent;
 
@@ -43,6 +45,9 @@ struct FClusteredComponentData
 	// Every physics object associated with this particular component.
 	TArray<Chaos::FPhysicsObjectHandle> AllPhysicsObjects;
 
+	// Cached acceleration structure handles - needed to properly cleanup the component from the accel structure.
+	TSet<FExternalSpatialAccelerationPayload> CachedAccelerationPayloads;
+
 	// Using a TWeakObjectPtr here because the UClusterUnionReplicatedProxyComponent will have a pointer back
 	// and we don't want to get into a situation where a circular reference occurs.
 	UPROPERTY()
@@ -92,6 +97,9 @@ struct FClusterUnionPendingAddData
 
 	UPROPERTY()
 	TArray<int32> BoneIds;
+	
+	UPROPERTY()
+	TArray<FExternalSpatialAccelerationPayload> AccelerationPayloads;
 };
 
 /**
@@ -239,7 +247,7 @@ private:
 	FChaosUserData PhysicsUserData;
 
 	// An acceleration structure of all children components managed by the cluster union itself.
-	TUniquePtr<Chaos::ISpatialAcceleration<Chaos::FAccelerationStructureHandle, Chaos::FReal, 3>> AccelerationStructure;
+	TUniquePtr<Chaos::ISpatialAcceleration<FExternalSpatialAccelerationPayload, Chaos::FReal, 3>> AccelerationStructure;
 
 	// Need to handle the fact that this component may or may not be initialized prior to the components referenced in
 	// ClusteredComponentsReferences. This function lets us listen to OnComponentPhysicsStateChanged on the incoming
diff --git a/Engine/Source/Runtime/Engine/Classes/PhysicsEngine/ExternalSpatialAccelerationPayload.h b/Engine/Source/Runtime/Engine/Classes/PhysicsEngine/ExternalSpatialAccelerationPayload.h
new file mode 100644
--- /dev/null
+++ b/Engine/Source/Runtime/Engine/Classes/PhysicsEngine/ExternalSpatialAccelerationPayload.h
@@ -0,0 +1,69 @@
+// Copyright Epic Games, Inc. All Rights Reserved.
+
+#pragma once
+
+#include "Chaos/ChaosArchive.h"
+#include "Chaos/Declares.h"
+#include "Chaos/GeometryParticlesfwd.h"
+#include "Chaos/ParticleHandleFwd.h"
+#include "Chaos/SpatialAccelerationFwd.h"
+#include "UObject/Object.h"
+#include "UObject/WeakObjectPtrTemplates.h"
+
+#include "ExternalSpatialAccelerationPayload.generated.h"
+
+class UPrimitiveComponent;
+
+/**
+ * This is a Chaos spatial acceleration payload that can be used for spatial acceleration structures that are meant to only be used
+ * in external contexts. This provides safety as we aren't holding particle handles directly but rather going through the primitive component's
+ * physics object interface to grab them when needed.
+ */
+USTRUCT()
+struct ENGINE_API FExternalSpatialAccelerationPayload
+{
+	GENERATED_BODY()
+public:
+	static constexpr bool bHasPayloadOnInternalThread = false;
+
+	FExternalSpatialAccelerationPayload();
+
+	void Initialize(UPrimitiveComponent* InComponent, int32 InBoneId);
+
+	Chaos::FGeometryParticle* GetExternalGeometryParticle_ExternalThread() const;
+	bool operator==(const FExternalSpatialAccelerationPayload& Other) const;
+	bool operator!=(const FExternalSpatialAccelerationPayload& Other) const;
+	Chaos::FUniqueIdx UniqueIdx() const;
+	void Serialize(Chaos::FChaosArchive& Ar);
+	bool PrePreQueryFilter(const void* QueryData) const { return false; }
+	bool PrePreSimFilter(const void* SimData) const { return false; }
+
+	friend ENGINE_API uint32 GetTypeHash(const FExternalSpatialAccelerationPayload& Payload);
+private:
+	TWeakObjectPtr<UPrimitiveComponent> Component = nullptr;
+	int32 BoneId = INDEX_NONE;
+	Chaos::FUniqueIdx CachedUniqueIdx;
+
+#if CHAOS_DEBUG_DRAW
+public:
+	void DebugDraw(const bool bExternal, const bool bHit) const;
+#endif
+};
+
+using IExternalSpatialAcceleration = Chaos::ISpatialAcceleration<FExternalSpatialAccelerationPayload, Chaos::FReal, 3>;
+
+FORCEINLINE bool FExternalSpatialAccelerationPayload::operator==(const FExternalSpatialAccelerationPayload& Other) const
+{
+	return (Component == Other.Component) && (BoneId == Other.BoneId);
+}
+
+FORCEINLINE bool FExternalSpatialAccelerationPayload::operator!=(const FExternalSpatialAccelerationPayload& Other) const
+{
+	return !(*this == Other);
+}
+
+FORCEINLINE_DEBUGGABLE Chaos::FChaosArchive& operator<<(Chaos::FChaosArchive& Ar, FExternalSpatialAccelerationPayload& AccelerationHandle)
+{
+	AccelerationHandle.Serialize(Ar);
+	return Ar;
+}
\ No newline at end of file
diff --git a/Engine/Source/Runtime/Engine/Private/Collision/SceneQuery.cpp b/Engine/Source/Runtime/Engine/Private/Collision/SceneQuery.cpp
--- a/Engine/Source/Runtime/Engine/Private/Collision/SceneQuery.cpp
+++ b/Engine/Source/Runtime/Engine/Private/Collision/SceneQuery.cpp
@@ -7,6 +7,7 @@
 #include "Collision/CollisionConversions.h"
 #include "PhysicsEngine/ClusterUnionComponent.h"
 #include "PhysicsEngine/CollisionQueryFilterCallback.h"
+#include "PhysicsEngine/ExternalSpatialAccelerationPayload.h"
 #include "PhysicsEngine/ScopedSQHitchRepeater.h"
 
 #include "Collision/CollisionDebugDrawing.h"
@@ -88,6 +89,33 @@ struct FGeomCollectionSQAdditionalInputs
 	const FQuat& GeomRot;
 };
 
+struct FPhysicsGeometrySQAdditionalInputs
+{
+	FPhysicsGeometrySQAdditionalInputs(const FPhysicsGeometry& InGeometry, const FQuat& InGeomRot)
+		: Collection(FChaosEngineInterface::GetGeometryCollection(InGeometry))
+		, GeomRot(InGeomRot)
+	{
+	}
+
+	const FPhysicsGeometry* GetGeometry() const
+	{
+		return &Collection.GetGeometry();
+	}
+
+	const FQuat* GetGeometryOrientation() const
+	{
+		return &GeomRot;
+	}
+
+	const FPhysicsGeometryCollection* GetCollisionShape() const
+	{
+		return &Collection;
+	}
+private:
+	const FPhysicsGeometryCollection Collection;
+	const FQuat& GeomRot;
+};
+
 struct FRaycastSQAdditionalInputs
 {
 	const FPhysicsGeometry* GetGeometry() const
@@ -106,6 +134,31 @@ struct FRaycastSQAdditionalInputs
 	}
 };
 
+
+template<typename TGeom>
+using TGeomSQInputs = std::conditional_t<std::is_same_v<TGeom, FPhysicsGeometryCollection>, FGeomCollectionSQAdditionalInputs,
+	std::conditional_t<std::is_same_v<TGeom, FCollisionShape>, FGeomSQAdditionalInputs,
+	std::conditional_t<std::is_same_v<TGeom, FPhysicsGeometry>, FPhysicsGeometrySQAdditionalInputs,
+	void>>>;
+
+template<typename TGeom>
+TGeomSQInputs<TGeom> GeomToSQInputs(const TGeom& Geom, const FQuat& Rot)
+{
+	static_assert(!std::is_same_v<TGeomSQInputs<TGeom>, void>, "Invalid geometry passed to SQ.");
+	if constexpr (std::is_same_v<TGeom, FPhysicsGeometryCollection>)
+	{
+		return FGeomCollectionSQAdditionalInputs(Geom, Rot);
+	}
+	else if constexpr (std::is_same_v<TGeom, FPhysicsGeometry>)
+	{
+		return FPhysicsGeometrySQAdditionalInputs(Geom, Rot);
+	}
+	else if constexpr (std::is_same_v<TGeom, FCollisionShape>)
+	{
+		return FGeomSQAdditionalInputs(Geom, Rot);
+	}
+}
+
 template <typename InHitType, ESweepOrRay InGeometryQuery, ESingleMultiOrTest InSingleMultiOrTest>
 struct TSQTraits
 {
@@ -342,9 +395,10 @@ struct FDefaultAccelContainer
 	}
 };
 
+template<typename TAccel>
 struct FOverrideAccelContainer
 {
-	explicit FOverrideAccelContainer(const Chaos::IDefaultChaosSpatialAcceleration& InSpatialAcceleration)
+	explicit FOverrideAccelContainer(const TAccel& InSpatialAcceleration)
 		: SpatialAcceleration(InSpatialAcceleration)
 	{}
 
@@ -353,13 +407,13 @@ struct FOverrideAccelContainer
 		return true;
 	}
 
-	const Chaos::IDefaultChaosSpatialAcceleration& GetSpatialAcceleration() const
+	const TAccel& GetSpatialAcceleration() const
 	{
 		return SpatialAcceleration;
 	}
 
 private:
-	const Chaos::IDefaultChaosSpatialAcceleration& SpatialAcceleration;
+	const TAccel& SpatialAcceleration;
 };
 
 template <typename Traits, typename TGeomInputs, typename TAccelContainer>
@@ -591,7 +645,8 @@ bool FGenericPhysicsInterface::RaycastTest(const UWorld* World, const FVector St
 	return TSceneCastCommon<TCastTraits, TPTCastTraits>(World, DummyHit, FRaycastSQAdditionalInputs(), Start, End, TraceChannel, Params, ResponseParams, ObjectParams);
 }
 
-bool FGenericPhysicsInterface::RaycastTestUsingSpatialAcceleration(const Chaos::IDefaultChaosSpatialAcceleration& Accel, const UWorld* World, const FVector Start, const FVector End, ECollisionChannel TraceChannel, const FCollisionQueryParams& Params, const FCollisionResponseParams& ResponseParams, const FCollisionObjectQueryParams& ObjectParams)
+template<typename TAccel>
+bool FGenericRaycastPhysicsInterfaceUsingSpatialAcceleration<TAccel>::RaycastTest(const TAccel& Accel, const UWorld* World, const FVector Start, const FVector End, ECollisionChannel TraceChannel, const FCollisionQueryParams& Params, const FCollisionResponseParams& ResponseParams, const FCollisionObjectQueryParams& ObjectParams)
 {
 	using namespace ChaosInterface;
 
@@ -619,7 +674,8 @@ bool FGenericPhysicsInterface::RaycastSingle(const UWorld* World, struct FHitRes
 	return TSceneCastCommon<TCastTraits, TPTCastTraits>(World, OutHit, FRaycastSQAdditionalInputs(), Start, End, TraceChannel, Params, ResponseParams, ObjectParams);
 }
 
-bool FGenericPhysicsInterface::RaycastSingleUsingSpatialAcceleration(const Chaos::IDefaultChaosSpatialAcceleration& Accel, const UWorld* World, struct FHitResult& OutHit, const FVector Start, const FVector End, ECollisionChannel TraceChannel, const FCollisionQueryParams& Params, const FCollisionResponseParams& ResponseParams, const FCollisionObjectQueryParams& ObjectParams)
+template<typename TAccel>
+bool FGenericRaycastPhysicsInterfaceUsingSpatialAcceleration<TAccel>::RaycastSingle(const TAccel& Accel, const UWorld* World, struct FHitResult& OutHit, const FVector Start, const FVector End, ECollisionChannel TraceChannel, const FCollisionQueryParams& Params, const FCollisionResponseParams& ResponseParams, const FCollisionObjectQueryParams& ObjectParams)
 {
 	using namespace ChaosInterface;
 
@@ -646,7 +702,8 @@ bool FGenericPhysicsInterface::RaycastMulti(const UWorld* World, TArray<struct F
 	return TSceneCastCommon<TCastTraits, TPTCastTraits> (World, OutHits, FRaycastSQAdditionalInputs(), Start, End, TraceChannel, Params, ResponseParams, ObjectParams);
 }
 
-bool FGenericPhysicsInterface::RaycastMultiUsingSpatialAcceleration(const Chaos::IDefaultChaosSpatialAcceleration& Accel, const UWorld* World, TArray<struct FHitResult>& OutHits, const FVector& Start, const FVector& End, ECollisionChannel TraceChannel, const struct FCollisionQueryParams& Params, const struct FCollisionResponseParams& ResponseParams, const struct FCollisionObjectQueryParams& ObjectParams)
+template<typename TAccel>
+bool FGenericRaycastPhysicsInterfaceUsingSpatialAcceleration<TAccel>::RaycastMulti(const TAccel& Accel, const UWorld* World, TArray<struct FHitResult>& OutHits, const FVector& Start, const FVector& End, ECollisionChannel TraceChannel, const struct FCollisionQueryParams& Params, const struct FCollisionResponseParams& ResponseParams, const struct FCollisionObjectQueryParams& ObjectParams)
 {
 	using namespace ChaosInterface;
 
@@ -677,7 +734,8 @@ bool FGenericPhysicsInterface::GeomSweepTest(const UWorld* World, const struct F
 	return TSceneCastCommon<TCastTraits, TPTCastTraits>(World, DummyHit, FGeomSQAdditionalInputs(CollisionShape, Rot), Start, End, TraceChannel, Params, ResponseParams, ObjectParams);
 }
 
-bool FGenericPhysicsInterface::GeomSweepTestUsingSpatialAcceleration(const Chaos::IDefaultChaosSpatialAcceleration& Accel, const UWorld* World, const struct FCollisionShape& CollisionShape, const FQuat& Rot, FVector Start, FVector End, ECollisionChannel TraceChannel, const struct FCollisionQueryParams& Params, const struct FCollisionResponseParams& ResponseParams, const struct FCollisionObjectQueryParams& ObjectParams)
+template<typename TAccel, typename TGeom>
+bool FGenericGeomPhysicsInterfaceUsingSpatialAcceleration<TAccel, TGeom>::GeomSweepTest(const TAccel& Accel, const UWorld* World, const TGeom& InGeom, const FQuat& Rot, FVector Start, FVector End, ECollisionChannel TraceChannel, const struct FCollisionQueryParams& Params, const struct FCollisionResponseParams& ResponseParams, const struct FCollisionObjectQueryParams& ObjectParams)
 {
 	using namespace ChaosInterface;
 
@@ -689,7 +747,7 @@ bool FGenericPhysicsInterface::GeomSweepTestUsingSpatialAcceleration(const Chaos
 	using TPTCastTraits = TSQTraits<FPTSweepHit, ESweepOrRay::Sweep, ESingleMultiOrTest::Test>;
 	FHitResult DummyHit(NoInit);
 
-	return TSceneCastCommon<TCastTraits, TPTCastTraits>(World, DummyHit, FGeomSQAdditionalInputs(CollisionShape, Rot), Start, End, TraceChannel, Params, ResponseParams, ObjectParams, FOverrideAccelContainer{ Accel });
+	return TSceneCastCommon<TCastTraits, TPTCastTraits>(World, DummyHit, GeomToSQInputs(InGeom, Rot), Start, End, TraceChannel, Params, ResponseParams, ObjectParams, FOverrideAccelContainer{ Accel });
 }
 
 template<>
@@ -946,16 +1004,16 @@ bool FGenericPhysicsInterface::GeomOverlapBlockingTest(const UWorld* World, cons
 	return GeomOverlapMultiHelper<EQueryInfo::IsBlocking>(World, Adaptor.GetGeometry(), CollisionShape, Adaptor.GetGeomPose(GeomTransform.GetTranslation()), Overlaps, TraceChannel, Params, ResponseParams, ObjectParams);
 }
 
-bool FGenericPhysicsInterface::GeomOverlapBlockingTestUsingSpatialAcceleration(const Chaos::IDefaultChaosSpatialAcceleration& Accel, const UWorld* World, const struct FCollisionShape& CollisionShape, const FVector& Pos, const FQuat& Rot, ECollisionChannel TraceChannel, const struct FCollisionQueryParams& Params, const struct FCollisionResponseParams& ResponseParams, const struct FCollisionObjectQueryParams& ObjectParams)
+template<typename TAccel, typename TGeom>
+bool FGenericGeomPhysicsInterfaceUsingSpatialAcceleration<TAccel, TGeom>::GeomOverlapBlockingTest(const TAccel& Accel, const UWorld* World, const TGeom& InGeom, const FVector& Pos, const FQuat& Rot, ECollisionChannel TraceChannel, const struct FCollisionQueryParams& Params, const struct FCollisionResponseParams& ResponseParams, const struct FCollisionObjectQueryParams& ObjectParams)
 {
 	SCOPE_CYCLE_COUNTER(STAT_Collision_SceneQueryTotal);
 	SCOPE_CYCLE_COUNTER(STAT_Collision_GeomOverlapBlocking);
 	CSV_SCOPED_TIMING_STAT(SceneQuery, GeomOverlapBlocking);
 
 	TArray<FOverlapResult> Overlaps;	//needed only for template shared code
-	FTransform GeomTransform(Rot, Pos);
-	FPhysicsShapeAdapter Adaptor(GeomTransform.GetRotation(), CollisionShape);
-	return GeomOverlapMultiHelper<EQueryInfo::IsBlocking>(World, Adaptor.GetGeometry(), CollisionShape, Adaptor.GetGeomPose(GeomTransform.GetTranslation()), Overlaps, TraceChannel, Params, ResponseParams, ObjectParams, FOverrideAccelContainer{ Accel });
+	TGeomSQInputs<TGeom> SQInputs = GeomToSQInputs(InGeom, Rot);
+	return GeomOverlapMultiHelper<EQueryInfo::IsBlocking>(World, *SQInputs.GetGeometry(), *SQInputs.GetCollisionShape(), FTransform{Rot, Pos}, Overlaps, TraceChannel, Params, ResponseParams, ObjectParams, FOverrideAccelContainer{Accel});
 }
 
 bool FGenericPhysicsInterface::GeomOverlapAnyTest(const UWorld* World, const struct FCollisionShape& CollisionShape, const FVector& Pos, const FQuat& Rot, ECollisionChannel TraceChannel, const struct FCollisionQueryParams& Params, const struct FCollisionResponseParams& ResponseParams, const struct FCollisionObjectQueryParams& ObjectParams)
@@ -970,16 +1028,16 @@ bool FGenericPhysicsInterface::GeomOverlapAnyTest(const UWorld* World, const str
 	return GeomOverlapMultiHelper<EQueryInfo::IsAnything>(World, Adaptor.GetGeometry(), CollisionShape, Adaptor.GetGeomPose(GeomTransform.GetTranslation()), Overlaps, TraceChannel, Params, ResponseParams, ObjectParams);
 }
 
-bool FGenericPhysicsInterface::GeomOverlapAnyTestUsingSpatialAcceleration(const Chaos::IDefaultChaosSpatialAcceleration& Accel, const UWorld* World, const struct FCollisionShape& CollisionShape, const FVector& Pos, const FQuat& Rot, ECollisionChannel TraceChannel, const struct FCollisionQueryParams& Params, const struct FCollisionResponseParams& ResponseParams, const struct FCollisionObjectQueryParams& ObjectParams)
+template<typename TAccel, typename TGeom>
+bool FGenericGeomPhysicsInterfaceUsingSpatialAcceleration<TAccel, TGeom>::GeomOverlapAnyTest(const TAccel& Accel, const UWorld* World, const TGeom& InGeom, const FVector& Pos, const FQuat& Rot, ECollisionChannel TraceChannel, const struct FCollisionQueryParams& Params, const struct FCollisionResponseParams& ResponseParams, const struct FCollisionObjectQueryParams& ObjectParams)
 {
 	SCOPE_CYCLE_COUNTER(STAT_Collision_SceneQueryTotal);
 	SCOPE_CYCLE_COUNTER(STAT_Collision_GeomOverlapAny);
 	CSV_SCOPED_TIMING_STAT(SceneQuery, GeomOverlapAny);
 
 	TArray<FOverlapResult> Overlaps;	//needed only for template shared code
-	FTransform GeomTransform(Rot, Pos);
-	FPhysicsShapeAdapter Adaptor(GeomTransform.GetRotation(), CollisionShape);
-	return GeomOverlapMultiHelper<EQueryInfo::IsAnything>(World, Adaptor.GetGeometry(), CollisionShape, Adaptor.GetGeomPose(GeomTransform.GetTranslation()), Overlaps, TraceChannel, Params, ResponseParams, ObjectParams, FOverrideAccelContainer{ Accel });
+	TGeomSQInputs<TGeom> SQInputs = GeomToSQInputs(InGeom, Rot);
+	return GeomOverlapMultiHelper<EQueryInfo::IsAnything>(World, *SQInputs.GetGeometry(), *SQInputs.GetCollisionShape(), FTransform{ Rot, Pos }, Overlaps, TraceChannel, Params, ResponseParams, ObjectParams, FOverrideAccelContainer{ Accel });
 }
 
 template<>
@@ -1011,8 +1069,8 @@ bool FGenericPhysicsInterface::GeomOverlapMulti(const UWorld* World, const FPhys
 	return GeomOverlapMulti(World, FChaosEngineInterface::GetGeometryCollection(InGeom), InPosition, InRotation, OutOverlaps, TraceChannel, Params, ResponseParams, ObjectParams);
 }
 
-template<>
-bool FGenericPhysicsInterface::GeomSweepSingleUsingSpatialAcceleration(const Chaos::IDefaultChaosSpatialAcceleration& Accel, const UWorld* World, const FPhysicsGeometryCollection& InGeom, const FQuat& Rot, FHitResult& OutHit, FVector Start, FVector End, ECollisionChannel TraceChannel, const FCollisionQueryParams& Params, const FCollisionResponseParams& ResponseParams, const FCollisionObjectQueryParams& ObjectParams)
+template<typename TAccel, typename TGeom>
+bool FGenericGeomPhysicsInterfaceUsingSpatialAcceleration<TAccel, TGeom>::GeomSweepSingle(const TAccel& Accel, const UWorld* World, const TGeom& InGeom, const FQuat& Rot, FHitResult& OutHit, FVector Start, FVector End, ECollisionChannel TraceChannel, const FCollisionQueryParams& Params, const FCollisionResponseParams& ResponseParams, const FCollisionObjectQueryParams& ObjectParams)
 {
 	using namespace ChaosInterface;
 
@@ -1022,31 +1080,11 @@ bool FGenericPhysicsInterface::GeomSweepSingleUsingSpatialAcceleration(const Cha
 
 	using TCastTraits = TSQTraits<FSweepHit, ESweepOrRay::Sweep, ESingleMultiOrTest::Single>;
 	using TPTCastTraits = TSQTraits<FPTSweepHit, ESweepOrRay::Sweep, ESingleMultiOrTest::Single>;
-	return TSceneCastCommon<TCastTraits, TPTCastTraits>(World, OutHit, FGeomCollectionSQAdditionalInputs(InGeom, Rot), Start, End, TraceChannel, Params, ResponseParams, ObjectParams, FOverrideAccelContainer{ Accel });
+	return TSceneCastCommon<TCastTraits, TPTCastTraits>(World, OutHit, GeomToSQInputs(InGeom, Rot), Start, End, TraceChannel, Params, ResponseParams, ObjectParams, FOverrideAccelContainer{ Accel });
 }
 
-template<>
-bool FGenericPhysicsInterface::GeomSweepSingleUsingSpatialAcceleration(const Chaos::IDefaultChaosSpatialAcceleration& Accel, const UWorld* World, const FCollisionShape& InGeom, const FQuat& Rot, FHitResult& OutHit, FVector Start, FVector End, ECollisionChannel TraceChannel, const FCollisionQueryParams& Params, const FCollisionResponseParams& ResponseParams, const FCollisionObjectQueryParams& ObjectParams)
-{
-	using namespace ChaosInterface;
-
-	SCOPE_CYCLE_COUNTER(STAT_Collision_SceneQueryTotal);
-	SCOPE_CYCLE_COUNTER(STAT_Collision_GeomSweepSingle);
-	CSV_SCOPED_TIMING_STAT(SceneQuery, GeomSweepSingle);
-
-	using TCastTraits = TSQTraits<FSweepHit, ESweepOrRay::Sweep, ESingleMultiOrTest::Single>;
-	using TPTCastTraits = TSQTraits<FPTSweepHit, ESweepOrRay::Sweep, ESingleMultiOrTest::Single>;
-	return TSceneCastCommon<TCastTraits, TPTCastTraits>(World, OutHit, FGeomSQAdditionalInputs(InGeom, Rot), Start, End, TraceChannel, Params, ResponseParams, ObjectParams, FOverrideAccelContainer{ Accel });
-}
-
-template<>
-bool FGenericPhysicsInterface::GeomSweepSingleUsingSpatialAcceleration(const Chaos::IDefaultChaosSpatialAcceleration& Accel, const UWorld* World, const FPhysicsGeometry& InGeom, const FQuat& Rot, FHitResult& OutHit, FVector Start, FVector End, ECollisionChannel TraceChannel, const FCollisionQueryParams& Params, const FCollisionResponseParams& ResponseParams, const FCollisionObjectQueryParams& ObjectParams)
-{
-	return GeomSweepSingleUsingSpatialAcceleration(Accel, World, FChaosEngineInterface::GetGeometryCollection(InGeom), Rot, OutHit, Start, End, TraceChannel, Params, ResponseParams, ObjectParams);
-}
-
-template<>
-bool FGenericPhysicsInterface::GeomSweepMultiUsingSpatialAcceleration(const Chaos::IDefaultChaosSpatialAcceleration& Accel, const UWorld* World, const FPhysicsGeometryCollection& InGeom, const FQuat& InGeomRot, TArray<FHitResult>& OutHits, FVector Start, FVector End, ECollisionChannel TraceChannel, const FCollisionQueryParams& Params, const FCollisionResponseParams& ResponseParams, const FCollisionObjectQueryParams& ObjectParams)
+template<typename TAccel, typename TGeom>
+bool FGenericGeomPhysicsInterfaceUsingSpatialAcceleration<TAccel, TGeom>::GeomSweepMulti(const TAccel& Accel, const UWorld* World, const TGeom& InGeom, const FQuat& InGeomRot, TArray<FHitResult>& OutHits, FVector Start, FVector End, ECollisionChannel TraceChannel, const FCollisionQueryParams& Params, const FCollisionResponseParams& ResponseParams, const FCollisionObjectQueryParams& ObjectParams)
 {
 	using namespace ChaosInterface;
 
@@ -1056,54 +1094,27 @@ bool FGenericPhysicsInterface::GeomSweepMultiUsingSpatialAcceleration(const Chao
 
 	using TCastTraits = TSQTraits<FHitSweep, ESweepOrRay::Sweep, ESingleMultiOrTest::Multi>;
 	using TPTCastTraits = TSQTraits<FPTSweepHit, ESweepOrRay::Sweep, ESingleMultiOrTest::Multi>;
-	return TSceneCastCommon<TCastTraits, TPTCastTraits>(World, OutHits, FGeomCollectionSQAdditionalInputs(InGeom, InGeomRot), Start, End, TraceChannel, Params, ResponseParams, ObjectParams, FOverrideAccelContainer{ Accel });
+	return TSceneCastCommon<TCastTraits, TPTCastTraits>(World, OutHits, GeomToSQInputs(InGeom, InGeomRot), Start, End, TraceChannel, Params, ResponseParams, ObjectParams, FOverrideAccelContainer{ Accel });
 }
 
-template<>
-bool FGenericPhysicsInterface::GeomSweepMultiUsingSpatialAcceleration(const Chaos::IDefaultChaosSpatialAcceleration& Accel, const UWorld* World, const FCollisionShape& InGeom, const FQuat& InGeomRot, TArray<FHitResult>& OutHits, FVector Start, FVector End, ECollisionChannel TraceChannel, const FCollisionQueryParams& Params, const FCollisionResponseParams& ResponseParams, const FCollisionObjectQueryParams& ObjectParams)
-{
-	using namespace ChaosInterface;
-
-	SCOPE_CYCLE_COUNTER(STAT_Collision_SceneQueryTotal);
-	SCOPE_CYCLE_COUNTER(STAT_Collision_GeomSweepMultiple);
-	CSV_SCOPED_TIMING_STAT(SceneQuery, GeomSweepMultiple);
-
-	using TCastTraits = TSQTraits<FHitSweep, ESweepOrRay::Sweep, ESingleMultiOrTest::Multi>;
-	using TPTCastTraits = TSQTraits<FPTSweepHit, ESweepOrRay::Sweep, ESingleMultiOrTest::Multi>;
-	return TSceneCastCommon<TCastTraits, TPTCastTraits>(World, OutHits, FGeomSQAdditionalInputs(InGeom, InGeomRot), Start, End, TraceChannel, Params, ResponseParams, ObjectParams, FOverrideAccelContainer{ Accel });
-}
-
-template<>
-bool FGenericPhysicsInterface::GeomSweepMultiUsingSpatialAcceleration(const Chaos::IDefaultChaosSpatialAcceleration& Accel, const UWorld* World, const FPhysicsGeometry& InGeom, const FQuat& InGeomRot, TArray<FHitResult>& OutHits, FVector Start, FVector End, ECollisionChannel TraceChannel, const FCollisionQueryParams& Params, const FCollisionResponseParams& ResponseParams, const FCollisionObjectQueryParams& ObjectParams)
-{
-	return GeomSweepMultiUsingSpatialAcceleration(Accel, World, FChaosEngineInterface::GetGeometryCollection(InGeom), InGeomRot, OutHits, Start, End, TraceChannel, Params, ResponseParams, ObjectParams);
-}
-
-template<>
-bool FGenericPhysicsInterface::GeomOverlapMultiUsingSpatialAcceleration(const Chaos::IDefaultChaosSpatialAcceleration& Accel, const UWorld* World, const FPhysicsGeometryCollection& InGeom, const FVector& InPosition, const FQuat& InRotation, TArray<FOverlapResult>& OutOverlaps, ECollisionChannel TraceChannel, const FCollisionQueryParams& Params, const FCollisionResponseParams& ResponseParams, const FCollisionObjectQueryParams& ObjectParams)
+template<typename TAccel, typename TGeom>
+bool FGenericGeomPhysicsInterfaceUsingSpatialAcceleration<TAccel, TGeom>::GeomOverlapMulti(const TAccel& Accel, const UWorld* World, const TGeom& InGeom, const FVector& InPosition, const FQuat& InRotation, TArray<FOverlapResult>& OutOverlaps, ECollisionChannel TraceChannel, const FCollisionQueryParams& Params, const FCollisionResponseParams& ResponseParams, const FCollisionObjectQueryParams& ObjectParams)
 {
 	SCOPE_CYCLE_COUNTER(STAT_Collision_SceneQueryTotal);
 	SCOPE_CYCLE_COUNTER(STAT_Collision_GeomOverlapMultiple);
 	CSV_SCOPED_TIMING_STAT(SceneQuery, GeomOverlapMultiple);
 
-	FTransform GeomTransform(InRotation, InPosition);
-	return GeomOverlapMultiHelper<EQueryInfo::GatherAll>(World, InGeom.GetGeometry(), InGeom, GeomTransform, OutOverlaps, TraceChannel, Params, ResponseParams, ObjectParams, FOverrideAccelContainer{ Accel });
+	TGeomSQInputs<TGeom> SQInputs = GeomToSQInputs(InGeom, InRotation);
+	return GeomOverlapMultiHelper<EQueryInfo::GatherAll>(World, *SQInputs.GetGeometry(), *SQInputs.GetCollisionShape(), FTransform{ InRotation, InPosition }, OutOverlaps, TraceChannel, Params, ResponseParams, ObjectParams, FOverrideAccelContainer{ Accel });
 }
 
-template<>
-bool FGenericPhysicsInterface::GeomOverlapMultiUsingSpatialAcceleration(const Chaos::IDefaultChaosSpatialAcceleration& Accel, const UWorld* World, const FCollisionShape& InGeom, const FVector& InPosition, const FQuat& InRotation, TArray<FOverlapResult>& OutOverlaps, ECollisionChannel TraceChannel, const FCollisionQueryParams& Params, const FCollisionResponseParams& ResponseParams, const FCollisionObjectQueryParams& ObjectParams)
-{
-	SCOPE_CYCLE_COUNTER(STAT_Collision_SceneQueryTotal);
-	SCOPE_CYCLE_COUNTER(STAT_Collision_GeomOverlapMultiple);
-	CSV_SCOPED_TIMING_STAT(SceneQuery, GeomOverlapMultiple);
+template struct FGenericRaycastPhysicsInterfaceUsingSpatialAcceleration<Chaos::IDefaultChaosSpatialAcceleration>;
+template struct FGenericRaycastPhysicsInterfaceUsingSpatialAcceleration<IExternalSpatialAcceleration>;
 
-	FTransform GeomTransform(InRotation, InPosition);
-	FPhysicsShapeAdapter Adaptor(GeomTransform.GetRotation(), InGeom);
-	return GeomOverlapMultiHelper<EQueryInfo::GatherAll>(World, Adaptor.GetGeometry(), InGeom, Adaptor.GetGeomPose(GeomTransform.GetTranslation()), OutOverlaps, TraceChannel, Params, ResponseParams, ObjectParams, FOverrideAccelContainer{ Accel });
-}
+#define DEFINE_GENERIC_GEOM_FOR_ACCEL(TACCEL) \
+template struct FGenericGeomPhysicsInterfaceUsingSpatialAcceleration<TACCEL, FCollisionShape>;\
+template struct FGenericGeomPhysicsInterfaceUsingSpatialAcceleration<TACCEL, FPhysicsGeometry>;\
+template struct FGenericGeomPhysicsInterfaceUsingSpatialAcceleration<TACCEL, FPhysicsGeometryCollection>;
 
-template<>
-bool FGenericPhysicsInterface::GeomOverlapMultiUsingSpatialAcceleration(const Chaos::IDefaultChaosSpatialAcceleration& Accel, const UWorld* World, const FPhysicsGeometry& InGeom, const FVector& InPosition, const FQuat& InRotation, TArray<FOverlapResult>& OutOverlaps, ECollisionChannel TraceChannel, const FCollisionQueryParams& Params, const FCollisionResponseParams& ResponseParams, const FCollisionObjectQueryParams& ObjectParams)
-{
-	return GeomOverlapMultiUsingSpatialAcceleration(Accel, World, FChaosEngineInterface::GetGeometryCollection(InGeom), InPosition, InRotation, OutOverlaps, TraceChannel, Params, ResponseParams, ObjectParams);
-}
\ No newline at end of file
+DEFINE_GENERIC_GEOM_FOR_ACCEL(Chaos::IDefaultChaosSpatialAcceleration)
+DEFINE_GENERIC_GEOM_FOR_ACCEL(IExternalSpatialAcceleration)
\ No newline at end of file
diff --git a/Engine/Source/Runtime/Engine/Private/Collision/SceneQueryLowLevel.cpp b/Engine/Source/Runtime/Engine/Private/Collision/SceneQueryLowLevel.cpp
--- a/Engine/Source/Runtime/Engine/Private/Collision/SceneQueryLowLevel.cpp
+++ b/Engine/Source/Runtime/Engine/Private/Collision/SceneQueryLowLevel.cpp
@@ -1,15 +1,14 @@
 // Copyright Epic Games, Inc. All Rights Reserved.
 
-
-
-
 #include "Physics/Experimental/ChaosInterfaceWrapper.h"
 #include "SQAccelerator.h"
+#include "SQVisitor.h"
 
 #include "PhysTestSerializer.h"
 
 #include "PBDRigidsSolver.h"
 #include "Physics/Experimental/PhysScene_Chaos.h"
+#include "PhysicsEngine/ExternalSpatialAccelerationPayload.h"
 
 int32 ForceStandardSQ = 0;
 FAutoConsoleVariableRef CVarForceStandardSQ(TEXT("p.ForceStandardSQ"), ForceStandardSQ, TEXT("If enabled, we force the standard scene query even if custom SQ structure is enabled"));
@@ -66,7 +65,54 @@ namespace
 #endif
 	}
 
-	void SweepSQCaptureHelper(float QueryDurationSeconds, const FChaosSQAccelerator& SQAccelerator, const FPhysScene& Scene, const FPhysicsGeometry& QueryGeom, const FTransform& StartTM, const FVector& Dir, float DeltaMag, const FPhysicsHitCallback<FHitSweep>& HitBuffer, EHitFlags OutputFlags, FQueryFlags QueryFlags, const FCollisionFilterData& Filter, const ChaosInterface::FQueryFilterData& QueryFilterData, ICollisionQueryFilterCallbackBase* QueryCallback, const ChaosInterface::FQueryDebugParams& DebugParams)
+	/**
+	 * This is meant as a replacement for FChaosSQAccelerator. There doesn't seem to be a need to expose this publicly since it'll just be used internally to allow
+	 * us to use more than one type of acceleration structure payload.
+	 */
+	template<typename TPayload>
+	class FGenericChaosSQAccelerator
+	{
+	public:
+		explicit FGenericChaosSQAccelerator(const Chaos::ISpatialAcceleration<TPayload, Chaos::FReal, 3>& InSpatialAcceleration)
+			: SpatialAcceleration(InSpatialAcceleration)
+		{}
+
+		template<typename THit>
+		void Raycast(const FVector& Start, const FVector& Dir, const float DeltaMagnitude, ChaosInterface::FSQHitBuffer<THit>& HitBuffer, EHitFlags OutputFlags, const ChaosInterface::FQueryFilterData& QueryFilterData, ICollisionQueryFilterCallbackBase& QueryCallback, const ChaosInterface::FQueryDebugParams& DebugParams = {}) const
+		{
+			using namespace Chaos;
+			using namespace ChaosInterface;
+
+			TSQVisitor<TSphere<FReal, 3>, TPayload, THit, std::is_same_v<THit, FRaycastHit>> RaycastVisitor(Start, Dir, HitBuffer, OutputFlags, QueryFilterData, QueryCallback, DebugParams);
+			HitBuffer.IncFlushCount();
+			SpatialAcceleration.Raycast(Start, Dir, DeltaMagnitude, RaycastVisitor);
+			HitBuffer.DecFlushCount();
+		}
+
+		template<typename THit>
+		void Sweep(const Chaos::FImplicitObject& QueryGeom, const FTransform& StartTM, const FVector& Dir, const float DeltaMagnitude, ChaosInterface::FSQHitBuffer<THit>& HitBuffer, EHitFlags OutputFlags, const ChaosInterface::FQueryFilterData& QueryFilterData, ICollisionQueryFilterCallbackBase& QueryCallback, const ChaosInterface::FQueryDebugParams& DebugParams = {}) const
+		{
+			return Chaos::Utilities::CastHelper(QueryGeom, StartTM, [&](const auto& Downcast, const FTransform& StartFullTM)
+			{
+				return SweepHelper(Downcast, SpatialAcceleration, StartFullTM, Dir, DeltaMagnitude, HitBuffer, OutputFlags, QueryFilterData, QueryCallback, DebugParams);
+			});
+		}
+
+		template<typename THit>
+		void Overlap(const Chaos::FImplicitObject& QueryGeom, const FTransform& GeomPose, ChaosInterface::FSQHitBuffer<THit>& HitBuffer, const ChaosInterface::FQueryFilterData& QueryFilterData, ICollisionQueryFilterCallbackBase& QueryCallback, const ChaosInterface::FQueryDebugParams& DebugParams = {}) const
+		{
+			return Chaos::Utilities::CastHelper(QueryGeom, GeomPose, [&](const auto& Downcast, const FTransform& GeomFullPose)
+			{
+				return OverlapHelper(Downcast, SpatialAcceleration, GeomFullPose, HitBuffer, QueryFilterData, QueryCallback, DebugParams);
+			});
+		}
+
+	private:
+		const Chaos::ISpatialAcceleration<TPayload, Chaos::FReal, 3>& SpatialAcceleration;
+	};
+
+	template<typename TPayload>
+	void SweepSQCaptureHelper(float QueryDurationSeconds, const FGenericChaosSQAccelerator<TPayload>& SQAccelerator, const FPhysScene& Scene, const FPhysicsGeometry& QueryGeom, const FTransform& StartTM, const FVector& Dir, float DeltaMag, const FPhysicsHitCallback<FHitSweep>& HitBuffer, EHitFlags OutputFlags, FQueryFlags QueryFlags, const FCollisionFilterData& Filter, const ChaosInterface::FQueryFilterData& QueryFilterData, ICollisionQueryFilterCallbackBase* QueryCallback, const ChaosInterface::FQueryDebugParams& DebugParams)
 	{
 #if !UE_BUILD_SHIPPING
 		float QueryDurationMicro = QueryDurationSeconds * 1000.0 * 1000.0;
@@ -102,7 +148,8 @@ namespace
 #endif
 	}
 
-	void RaycastSQCaptureHelper(float QueryDurationSeconds, const FChaosSQAccelerator& SQAccelerator, const FPhysScene& Scene, const FVector& Start, const FVector& Dir, float DeltaMag, const FPhysicsHitCallback<FHitRaycast>& HitBuffer, EHitFlags OutputFlags, FQueryFlags QueryFlags, const FCollisionFilterData& Filter, const ChaosInterface::FQueryFilterData& QueryFilterData, ICollisionQueryFilterCallbackBase* QueryCallback, const ChaosInterface::FQueryDebugParams& DebugParams)
+	template<typename TPayload>
+	void RaycastSQCaptureHelper(float QueryDurationSeconds, const FGenericChaosSQAccelerator<TPayload>& SQAccelerator, const FPhysScene& Scene, const FVector& Start, const FVector& Dir, float DeltaMag, const FPhysicsHitCallback<FHitRaycast>& HitBuffer, EHitFlags OutputFlags, FQueryFlags QueryFlags, const FCollisionFilterData& Filter, const ChaosInterface::FQueryFilterData& QueryFilterData, ICollisionQueryFilterCallbackBase* QueryCallback, const ChaosInterface::FQueryDebugParams& DebugParams)
 	{
 #if !UE_BUILD_SHIPPING
 		float QueryDurationMicro = QueryDurationSeconds * 1000.0 * 1000.0;
@@ -138,7 +185,8 @@ namespace
 #endif
 	}
 
-	void OverlapSQCaptureHelper(float QueryDurationSeconds, const FChaosSQAccelerator& SQAccelerator, const FPhysScene& Scene, const FPhysicsGeometry& QueryGeom, const FTransform& GeomPose, const FPhysicsHitCallback<FHitOverlap>& HitBuffer, FQueryFlags QueryFlags, const FCollisionFilterData& Filter, const ChaosInterface::FQueryFilterData& QueryFilterData, ICollisionQueryFilterCallbackBase* QueryCallback, const ChaosInterface::FQueryDebugParams& DebugParams)
+	template<typename TPayload>
+	void OverlapSQCaptureHelper(float QueryDurationSeconds, const FGenericChaosSQAccelerator<TPayload>& SQAccelerator, const FPhysScene& Scene, const FPhysicsGeometry& QueryGeom, const FTransform& GeomPose, const FPhysicsHitCallback<FHitOverlap>& HitBuffer, FQueryFlags QueryFlags, const FCollisionFilterData& Filter, const ChaosInterface::FQueryFilterData& QueryFilterData, ICollisionQueryFilterCallbackBase* QueryCallback, const ChaosInterface::FQueryDebugParams& DebugParams)
 	{
 #if !UE_BUILD_SHIPPING
 		float QueryDurationMicro = QueryDurationSeconds * 1000.0 * 1000.0;
@@ -177,7 +225,9 @@ namespace
 	template<typename TAccelContainer, bool bGTData>
 	struct FAccelerationContainerTraits
 	{
-		static const Chaos::IDefaultChaosSpatialAcceleration* GetSpatialAccelerationFromContainer(const TAccelContainer& Container)
+		using TAccelStructure = std::conditional_t<std::is_same_v<TAccelContainer, FPhysScene>, Chaos::IDefaultChaosSpatialAcceleration, TAccelContainer>;
+
+		static const TAccelStructure* GetSpatialAccelerationFromContainer(const TAccelContainer& Container)
 		{
 			if constexpr (IsPhysScene())
 			{
@@ -205,20 +255,26 @@ namespace
 	};
 }
 
+#define DEFINE_LOW_LEVEL_RAYCAST_ACCEL_HIT(TACCEL, THIT) template void LowLevelRaycast<TACCEL, THIT>(const TACCEL& Container, const FVector& Start, const FVector& Dir, float DeltaMag, FPhysicsHitCallback<THIT>& HitBuffer, EHitFlags OutputFlags, FQueryFlags QueryFlags, const FCollisionFilterData& Filter, const ChaosInterface::FQueryFilterData& QueryFilterData, ICollisionQueryFilterCallbackBase* QueryCallback, const ChaosInterface::FQueryDebugParams& DebugParams);
+#define DEFINE_LOW_LEVEL_RAYCAST_ACCEL(TACCEL) \
+DEFINE_LOW_LEVEL_RAYCAST_ACCEL_HIT(TACCEL, FHitRaycast); \
+DEFINE_LOW_LEVEL_RAYCAST_ACCEL_HIT(TACCEL, ChaosInterface::FPTRaycastHit);
+
 template <typename TAccelContainer, typename THitRaycast>
-void LowLevelRaycastImp(const TAccelContainer& Container, const FVector& Start, const FVector& Dir, float DeltaMag, FPhysicsHitCallback<THitRaycast>& HitBuffer, EHitFlags OutputFlags, FQueryFlags QueryFlags, const FCollisionFilterData& Filter, const ChaosInterface::FQueryFilterData& QueryFilterData, ICollisionQueryFilterCallbackBase* QueryCallback, const ChaosInterface::FQueryDebugParams& DebugParams)
+void LowLevelRaycast(const TAccelContainer& Container, const FVector& Start, const FVector& Dir, float DeltaMag, FPhysicsHitCallback<THitRaycast>& HitBuffer, EHitFlags OutputFlags, FQueryFlags QueryFlags, const FCollisionFilterData& Filter, const ChaosInterface::FQueryFilterData& QueryFilterData, ICollisionQueryFilterCallbackBase* QueryCallback, const ChaosInterface::FQueryDebugParams& DebugParams)
 {
 	constexpr bool bGTData = std::is_same_v<THitRaycast, FHitRaycast>;
-	if (const Chaos::IDefaultChaosSpatialAcceleration* SolverAccelerationStructure = FAccelerationContainerTraits<TAccelContainer, bGTData>::GetSpatialAccelerationFromContainer(Container))
+	using TTraits = FAccelerationContainerTraits<TAccelContainer, bGTData>;
+	if (const typename TTraits::TAccelStructure* SolverAccelerationStructure = TTraits::GetSpatialAccelerationFromContainer(Container))
 	{
-		FChaosSQAccelerator SQAccelerator(*SolverAccelerationStructure);
+		FGenericChaosSQAccelerator<typename TTraits::TAccelStructure::TPayload> SQAccelerator(*SolverAccelerationStructure);
 		double Time = 0.0;
 		{
 			FScopedDurationTimer Timer(Time);
 			SQAccelerator.Raycast(Start, Dir, DeltaMag, HitBuffer, OutputFlags, QueryFilterData, *QueryCallback, DebugParams);
 		}
 
-		if constexpr (bGTData && FAccelerationContainerTraits<TAccelContainer, bGTData>::IsPhysScene())
+		if constexpr (bGTData && TTraits::IsPhysScene())
 		{
 			if (!!SerializeSQs && !!EnableRaycastSQCapture)
 			{
@@ -228,33 +284,23 @@ void LowLevelRaycastImp(const TAccelContainer& Container, const FVector& Start,
 	}
 }
 
-void LowLevelRaycast(const FPhysScene& Scene, const FVector& Start, const FVector& Dir, float DeltaMag, FPhysicsHitCallback<FHitRaycast>& HitBuffer, EHitFlags OutputFlags, FQueryFlags QueryFlags, const FCollisionFilterData& Filter, const ChaosInterface::FQueryFilterData& QueryFilterData, ICollisionQueryFilterCallbackBase* QueryCallback, const ChaosInterface::FQueryDebugParams& DebugParams)
-{
-	LowLevelRaycastImp(Scene, Start, Dir, DeltaMag, HitBuffer, OutputFlags, QueryFlags, Filter, QueryFilterData, QueryCallback, DebugParams);
-}
+DEFINE_LOW_LEVEL_RAYCAST_ACCEL(Chaos::IDefaultChaosSpatialAcceleration)
+DEFINE_LOW_LEVEL_RAYCAST_ACCEL(IExternalSpatialAcceleration)
+DEFINE_LOW_LEVEL_RAYCAST_ACCEL(FPhysScene)
 
-void LowLevelRaycast(const FPhysScene& Scene, const FVector& Start, const FVector& Dir, float DeltaMag, FPhysicsHitCallback<ChaosInterface::FPTRaycastHit>& HitBuffer, EHitFlags OutputFlags, FQueryFlags QueryFlags, const FCollisionFilterData& Filter, const ChaosInterface::FQueryFilterData& QueryFilterData, ICollisionQueryFilterCallbackBase* QueryCallback, const ChaosInterface::FQueryDebugParams& DebugParams)
-{
-	LowLevelRaycastImp(Scene, Start, Dir, DeltaMag, HitBuffer, OutputFlags, QueryFlags, Filter, QueryFilterData, QueryCallback, DebugParams);
-}
-
-void LowLevelRaycast(const Chaos::IDefaultChaosSpatialAcceleration& AccelerationStructure, const FVector& Start, const FVector& Dir, float DeltaMag, FPhysicsHitCallback<FHitRaycast>& HitBuffer, EHitFlags OutputFlags, FQueryFlags QueryFlags, const FCollisionFilterData& Filter, const ChaosInterface::FQueryFilterData& QueryFilterData, ICollisionQueryFilterCallbackBase* QueryCallback, const ChaosInterface::FQueryDebugParams& DebugParams)
-{
-	LowLevelRaycastImp(AccelerationStructure, Start, Dir, DeltaMag, HitBuffer, OutputFlags, QueryFlags, Filter, QueryFilterData, QueryCallback, DebugParams);
-}
-
-void LowLevelRaycast(const Chaos::IDefaultChaosSpatialAcceleration& AccelerationStructure, const FVector& Start, const FVector& Dir, float DeltaMag, FPhysicsHitCallback<ChaosInterface::FPTRaycastHit>& HitBuffer, EHitFlags OutputFlags, FQueryFlags QueryFlags, const FCollisionFilterData& Filter, const ChaosInterface::FQueryFilterData& QueryFilterData, ICollisionQueryFilterCallbackBase* QueryCallback, const ChaosInterface::FQueryDebugParams& DebugParams)
-{
-	LowLevelRaycastImp(AccelerationStructure, Start, Dir, DeltaMag, HitBuffer, OutputFlags, QueryFlags, Filter, QueryFilterData, QueryCallback, DebugParams);
-}
+#define DEFINE_LOW_LEVEL_SWEEP_ACCEL_HIT(TACCEL, THIT) template void LowLevelSweep<TACCEL, THIT>(const TACCEL& Container, const FPhysicsGeometry& QueryGeom, const FTransform& StartTM, const FVector& Dir, float DeltaMag, FPhysicsHitCallback<THIT>& HitBuffer, EHitFlags OutputFlags, FQueryFlags QueryFlags, const FCollisionFilterData& Filter, const ChaosInterface::FQueryFilterData& QueryFilterData, ICollisionQueryFilterCallbackBase* QueryCallback, const ChaosInterface::FQueryDebugParams& DebugParams)
+#define DEFINE_LOW_LEVEL_SWEEP_ACCEL(TACCEL) \
+DEFINE_LOW_LEVEL_SWEEP_ACCEL_HIT(TACCEL, FHitSweep); \
+DEFINE_LOW_LEVEL_SWEEP_ACCEL_HIT(TACCEL, ChaosInterface::FPTSweepHit);
 
 template <typename TAccelContainer, typename THitSweep>
-void LowLevelSweepImp(const TAccelContainer& Container, const FPhysicsGeometry& QueryGeom, const FTransform& StartTM, const FVector& Dir, float DeltaMag, FPhysicsHitCallback<THitSweep>& HitBuffer, EHitFlags OutputFlags, FQueryFlags QueryFlags, const FCollisionFilterData& Filter, const ChaosInterface::FQueryFilterData& QueryFilterData, ICollisionQueryFilterCallbackBase* QueryCallback, const ChaosInterface::FQueryDebugParams& DebugParams)
+void LowLevelSweep(const TAccelContainer& Container, const FPhysicsGeometry& QueryGeom, const FTransform& StartTM, const FVector& Dir, float DeltaMag, FPhysicsHitCallback<THitSweep>& HitBuffer, EHitFlags OutputFlags, FQueryFlags QueryFlags, const FCollisionFilterData& Filter, const ChaosInterface::FQueryFilterData& QueryFilterData, ICollisionQueryFilterCallbackBase* QueryCallback, const ChaosInterface::FQueryDebugParams& DebugParams)
 {
 	constexpr bool bGTData = std::is_same<THitSweep, FHitSweep>::value;
-	if (const Chaos::IDefaultChaosSpatialAcceleration* SolverAccelerationStructure = FAccelerationContainerTraits<TAccelContainer, bGTData>::GetSpatialAccelerationFromContainer(Container))
+	using TTraits = FAccelerationContainerTraits<TAccelContainer, bGTData>;
+	if (const typename TTraits::TAccelStructure* SolverAccelerationStructure = TTraits::GetSpatialAccelerationFromContainer(Container))
 	{
-		FChaosSQAccelerator SQAccelerator(*SolverAccelerationStructure);
+		FGenericChaosSQAccelerator<typename TTraits::TAccelStructure::TPayload> SQAccelerator(*SolverAccelerationStructure);
 		{
 			double Time = 0.0;
 			{
@@ -273,33 +319,23 @@ void LowLevelSweepImp(const TAccelContainer& Container, const FPhysicsGeometry&
 	}
 }
 
-void LowLevelSweep(const FPhysScene& Scene, const FPhysicsGeometry& QueryGeom, const FTransform& StartTM, const FVector& Dir, float DeltaMag, FPhysicsHitCallback<FHitSweep>& HitBuffer, EHitFlags OutputFlags, FQueryFlags QueryFlags, const FCollisionFilterData& Filter, const ChaosInterface::FQueryFilterData& QueryFilterData, ICollisionQueryFilterCallbackBase* QueryCallback, const ChaosInterface::FQueryDebugParams& DebugParams)
-{
-	LowLevelSweepImp(Scene, QueryGeom, StartTM, Dir, DeltaMag, HitBuffer, OutputFlags, QueryFlags, Filter, QueryFilterData, QueryCallback, DebugParams);
-}
+DEFINE_LOW_LEVEL_SWEEP_ACCEL(Chaos::IDefaultChaosSpatialAcceleration)
+DEFINE_LOW_LEVEL_SWEEP_ACCEL(IExternalSpatialAcceleration)
+DEFINE_LOW_LEVEL_SWEEP_ACCEL(FPhysScene)
 
-void LowLevelSweep(const FPhysScene& Scene, const FPhysicsGeometry& QueryGeom, const FTransform& StartTM, const FVector& Dir, float DeltaMag, FPhysicsHitCallback<ChaosInterface::FPTSweepHit>& HitBuffer, EHitFlags OutputFlags, FQueryFlags QueryFlags, const FCollisionFilterData& Filter, const ChaosInterface::FQueryFilterData& QueryFilterData, ICollisionQueryFilterCallbackBase* QueryCallback, const ChaosInterface::FQueryDebugParams& DebugParams)
-{
-	LowLevelSweepImp(Scene, QueryGeom, StartTM, Dir, DeltaMag, HitBuffer, OutputFlags, QueryFlags, Filter, QueryFilterData, QueryCallback, DebugParams);
-}
-
-void LowLevelSweep(const Chaos::IDefaultChaosSpatialAcceleration& AccelerationStructure, const FPhysicsGeometry& Geom, const FTransform& StartTM, const FVector& Dir, float DeltaMag, FPhysicsHitCallback<FHitSweep>& HitBuffer, EHitFlags OutputFlags, FQueryFlags QueryFlags, const FCollisionFilterData& Filter, const ChaosInterface::FQueryFilterData& QueryFilterData, ICollisionQueryFilterCallbackBase* QueryCallback, const ChaosInterface::FQueryDebugParams& DebugParams)
-{
-	LowLevelSweepImp(AccelerationStructure, Geom, StartTM, Dir, DeltaMag, HitBuffer, OutputFlags, QueryFlags, Filter, QueryFilterData, QueryCallback, DebugParams);
-}
-
-void LowLevelSweep(const Chaos::IDefaultChaosSpatialAcceleration& AccelerationStructure, const FPhysicsGeometry& Geom, const FTransform& StartTM, const FVector& Dir, float DeltaMag, FPhysicsHitCallback<ChaosInterface::FPTSweepHit>& HitBuffer, EHitFlags OutputFlags, FQueryFlags QueryFlags, const FCollisionFilterData& Filter, const ChaosInterface::FQueryFilterData& QueryFilterData, ICollisionQueryFilterCallbackBase* QueryCallback, const ChaosInterface::FQueryDebugParams& DebugParams)
-{
-	LowLevelSweepImp(AccelerationStructure, Geom, StartTM, Dir, DeltaMag, HitBuffer, OutputFlags, QueryFlags, Filter, QueryFilterData, QueryCallback, DebugParams);
-}
+#define DEFINE_LOW_LEVEL_OVERLAP_ACCEL_HIT(TACCEL, THIT) template void LowLevelOverlap<TACCEL, THIT>(const TACCEL& Container, const FPhysicsGeometry& QueryGeom, const FTransform& GeomPose, FPhysicsHitCallback<THIT>& HitBuffer, FQueryFlags QueryFlags, const FCollisionFilterData& Filter, const ChaosInterface::FQueryFilterData& QueryFilterData, ICollisionQueryFilterCallbackBase* QueryCallback, const ChaosInterface::FQueryDebugParams& DebugParams)
+#define DEFINE_LOW_LEVEL_OVERLAP_ACCEL(TACCEL) \
+DEFINE_LOW_LEVEL_OVERLAP_ACCEL_HIT(TACCEL, FHitOverlap); \
+DEFINE_LOW_LEVEL_OVERLAP_ACCEL_HIT(TACCEL, ChaosInterface::FPTOverlapHit);
 
 template <typename TAccelContainer, typename THitOverlap>
-void LowLevelOverlapImp(const TAccelContainer& Container, const FPhysicsGeometry& QueryGeom, const FTransform& GeomPose, FPhysicsHitCallback<THitOverlap>& HitBuffer, FQueryFlags QueryFlags, const FCollisionFilterData& Filter, const ChaosInterface::FQueryFilterData& QueryFilterData, ICollisionQueryFilterCallbackBase* QueryCallback, const ChaosInterface::FQueryDebugParams& DebugParams)
+void LowLevelOverlap(const TAccelContainer& Container, const FPhysicsGeometry& QueryGeom, const FTransform& GeomPose, FPhysicsHitCallback<THitOverlap>& HitBuffer, FQueryFlags QueryFlags, const FCollisionFilterData& Filter, const ChaosInterface::FQueryFilterData& QueryFilterData, ICollisionQueryFilterCallbackBase* QueryCallback, const ChaosInterface::FQueryDebugParams& DebugParams)
 {
 	constexpr bool bGTData = std::is_same<THitOverlap, FHitOverlap>::value;
-	if (const Chaos::IDefaultChaosSpatialAcceleration* SolverAccelerationStructure = FAccelerationContainerTraits<TAccelContainer, bGTData>::GetSpatialAccelerationFromContainer(Container))
+	using TTraits = FAccelerationContainerTraits<TAccelContainer, bGTData>;
+	if (const typename TTraits::TAccelStructure* SolverAccelerationStructure = TTraits::GetSpatialAccelerationFromContainer(Container))
 	{
-		FChaosSQAccelerator SQAccelerator(*SolverAccelerationStructure);
+		FGenericChaosSQAccelerator<typename TTraits::TAccelStructure::TPayload> SQAccelerator(*SolverAccelerationStructure);
 		double Time = 0.0;
 		{
 			FScopedDurationTimer Timer(Time);
@@ -316,22 +352,6 @@ void LowLevelOverlapImp(const TAccelContainer& Container, const FPhysicsGeometry
 	}
 }
 
-void LowLevelOverlap(const FPhysScene& Scene, const FPhysicsGeometry& QueryGeom, const FTransform& GeomPose, FPhysicsHitCallback<FHitOverlap>& HitBuffer, FQueryFlags QueryFlags, const FCollisionFilterData& Filter, const ChaosInterface::FQueryFilterData& QueryFilterData, ICollisionQueryFilterCallbackBase* QueryCallback, const ChaosInterface::FQueryDebugParams& DebugParams)
-{
-	LowLevelOverlapImp(Scene, QueryGeom, GeomPose, HitBuffer, QueryFlags, Filter, QueryFilterData, QueryCallback, DebugParams);
-}
-
-void LowLevelOverlap(const FPhysScene& Scene, const FPhysicsGeometry& QueryGeom, const FTransform& GeomPose, FPhysicsHitCallback<ChaosInterface::FPTOverlapHit>& HitBuffer, FQueryFlags QueryFlags, const FCollisionFilterData& Filter, const ChaosInterface::FQueryFilterData& QueryFilterData, ICollisionQueryFilterCallbackBase* QueryCallback, const ChaosInterface::FQueryDebugParams& DebugParams)
-{
-	LowLevelOverlapImp(Scene, QueryGeom, GeomPose, HitBuffer, QueryFlags, Filter, QueryFilterData, QueryCallback, DebugParams);
-}
-
-void LowLevelOverlap(const Chaos::IDefaultChaosSpatialAcceleration& AccelerationStructure, const FPhysicsGeometry& Geom, const FTransform& GeomPose, FPhysicsHitCallback<FHitOverlap>& HitBuffer, FQueryFlags QueryFlags, const FCollisionFilterData& Filter, const ChaosInterface::FQueryFilterData& QueryFilterData, ICollisionQueryFilterCallbackBase* QueryCallback, const ChaosInterface::FQueryDebugParams& DebugParams)
-{
-	LowLevelOverlapImp(AccelerationStructure, Geom, GeomPose, HitBuffer, QueryFlags, Filter, QueryFilterData, QueryCallback, DebugParams);
-}
-
-void LowLevelOverlap(const Chaos::IDefaultChaosSpatialAcceleration& AccelerationStructure, const FPhysicsGeometry& Geom, const FTransform& GeomPose, FPhysicsHitCallback<ChaosInterface::FPTOverlapHit>& HitBuffer, FQueryFlags QueryFlags, const FCollisionFilterData& Filter, const ChaosInterface::FQueryFilterData& QueryFilterData, ICollisionQueryFilterCallbackBase* QueryCallback, const ChaosInterface::FQueryDebugParams& DebugParams)
-{
-	LowLevelOverlapImp(AccelerationStructure, Geom, GeomPose, HitBuffer, QueryFlags, Filter, QueryFilterData, QueryCallback, DebugParams);
-}
\ No newline at end of file
+DEFINE_LOW_LEVEL_OVERLAP_ACCEL(Chaos::IDefaultChaosSpatialAcceleration)
+DEFINE_LOW_LEVEL_OVERLAP_ACCEL(IExternalSpatialAcceleration)
+DEFINE_LOW_LEVEL_OVERLAP_ACCEL(FPhysScene)
\ No newline at end of file
diff --git a/Engine/Source/Runtime/Engine/Private/PhysicsEngine/ClusterUnionComponent.cpp b/Engine/Source/Runtime/Engine/Private/PhysicsEngine/ClusterUnionComponent.cpp
--- a/Engine/Source/Runtime/Engine/Private/PhysicsEngine/ClusterUnionComponent.cpp
+++ b/Engine/Source/Runtime/Engine/Private/PhysicsEngine/ClusterUnionComponent.cpp
@@ -10,6 +10,7 @@
 #include "Physics/Experimental/PhysScene_Chaos.h"
 #include "Physics/GenericPhysicsInterface.h"
 #include "PhysicsEngine/ClusterUnionReplicatedProxyComponent.h"
+#include "PhysicsEngine/ExternalSpatialAccelerationPayload.h"
 #include "PhysicsEngine/PhysicsObjectExternalInterface.h"
 #include "PhysicsProxy/ClusterUnionPhysicsProxy.h"
 #include "TimerManager.h"
@@ -24,13 +25,12 @@ namespace
 	FAutoConsoleVariableRef CVarUseClusterUnionAccelerationStructure(TEXT("ClusterUnion.UseAccelerationStructure"), bUseClusterUnionAccelerationStructure, TEXT("Whether component level sweeps and overlaps against cluster unions should use an acceleration structure instead."));
 
 	// TODO: Should this be exposed in Chaos instead?
-	using FAccelerationStructure = Chaos::TAABBTree<Chaos::FAccelerationStructureHandle, Chaos::TAABBTreeLeafArray<Chaos::FAccelerationStructureHandle>>;
+	using FAccelerationStructure = Chaos::TAABBTree<FExternalSpatialAccelerationPayload, Chaos::TAABBTreeLeafArray<FExternalSpatialAccelerationPayload>>;
 
-	TUniquePtr<Chaos::ISpatialAcceleration<Chaos::FAccelerationStructureHandle, Chaos::FReal, 3>> CreateEmptyAccelerationStructure()
+	TUniquePtr<Chaos::ISpatialAcceleration<FExternalSpatialAccelerationPayload, Chaos::FReal, 3>> CreateEmptyAccelerationStructure()
 	{
-		const Chaos::TConstParticleView<Chaos::FSpatialAccelerationCache> Empty;
 		FAccelerationStructure* Structure = new FAccelerationStructure(
-			Empty,
+			nullptr,
 			FAccelerationStructure::DefaultMaxChildrenInLeaf,
 			FAccelerationStructure::DefaultMaxTreeDepth,
 			FAccelerationStructure::DefaultMaxPayloadBounds,
@@ -40,7 +40,7 @@ namespace
 		);
 
 		check(Structure != nullptr);
-		return TUniquePtr<Chaos::ISpatialAcceleration<Chaos::FAccelerationStructureHandle, Chaos::FReal, 3>>(Structure);
+		return TUniquePtr<Chaos::ISpatialAcceleration<FExternalSpatialAccelerationPayload, Chaos::FReal, 3>>(Structure);
 	}
 }
 
@@ -139,16 +139,20 @@ void UClusterUnionComponent::AddComponentToCluster(UPrimitiveComponent* InCompon
 	}
 
 	PendingData.BoneIds.Empty(Objects.Num());
+	PendingData.AccelerationPayloads.Empty(Objects.Num());
 
 	for (Chaos::FPhysicsObjectHandle Object : Objects)
 	{
-		PendingData.BoneIds.Add(Chaos::FPhysicsObjectInterface::GetId(Object));
+		const int32 BoneId = Chaos::FPhysicsObjectInterface::GetId(Object);
+		PendingData.BoneIds.Add(BoneId);
 
 		if (AccelerationStructure && Object)
 		{
-			Chaos::FAccelerationStructureHandle Handle = Interface->CreateAccelerationStructureHandle(Object);
 			const FBox HandleBounds = Interface->GetWorldBounds({ &Object, 1 });
+			FExternalSpatialAccelerationPayload Handle;
+			Handle.Initialize(InComponent, BoneId);
 			AccelerationStructure->UpdateElement(Handle, Chaos::TAABB<Chaos::FReal, 3>{HandleBounds.Min, HandleBounds.Max}, HandleBounds.IsValid != 0);
+			PendingData.AccelerationPayloads.Add(Handle);
 		}
 	}
 
@@ -175,6 +179,20 @@ void UClusterUnionComponent::RemoveComponentFromCluster(UPrimitiveComponent* InC
 		return;
 	}
 
+	// If we're still waiting on the physics sync for the added component, we still need to
+	// remove its acceleration handles from the acceleration structure since the payloads have
+	// already been added and will not yet be stored in ComponentToPhysicsObjects.
+	if (FClusterUnionPendingAddData* PendingData = PendingComponentSync.Find(InComponent))
+	{
+		if (AccelerationStructure)
+		{
+			for (const FExternalSpatialAccelerationPayload& Payload : PendingData->AccelerationPayloads)
+			{
+				AccelerationStructure->RemoveElement(Payload);
+			}
+		}
+	}
+
 	TSet<Chaos::FPhysicsObjectHandle> PhysicsObjectsToRemove;
 	FLockedReadPhysicsObjectExternalInterface Interface = FPhysicsObjectExternalInterface::LockRead(GetWorld()->GetPhysicsScene());
 
@@ -196,7 +214,6 @@ void UClusterUnionComponent::RemoveComponentFromCluster(UPrimitiveComponent* InC
 		{
 			ComponentData->PhysicsObjects.Reset();
 			ComponentData->AllPhysicsObjects.Reset();
-			return;
 		}
 		else
 		{
@@ -205,19 +222,19 @@ void UClusterUnionComponent::RemoveComponentFromCluster(UPrimitiveComponent* InC
 
 		if (AccelerationStructure)
 		{
-			for (Chaos::FPhysicsObjectHandle PhysicsObject : PhysicsObjectsToRemove)
+			for (const FExternalSpatialAccelerationPayload& Payload : ComponentData->CachedAccelerationPayloads)
 			{
-				if (PhysicsObject)
-				{
-					Chaos::FAccelerationStructureHandle Handle = Interface->CreateAccelerationStructureHandle(PhysicsObject);
-					AccelerationStructure->RemoveElement(Handle);
-				}
+				AccelerationStructure->RemoveElement(Payload);
 			}
+			ComponentData->CachedAccelerationPayloads.Reset();
 		}
 	}
 
-	PhysicsProxy->RemovePhysicsObjects_External(PhysicsObjectsToRemove);
-	ForceRebuildGTParticleGeometry();
+	if (!PhysicsObjectsToRemove.IsEmpty())
+	{
+		PhysicsProxy->RemovePhysicsObjects_External(PhysicsObjectsToRemove);
+		ForceRebuildGTParticleGeometry();
+	}
 }
 
 void UClusterUnionComponent::ForceRebuildGTParticleGeometry()
@@ -575,7 +592,13 @@ void UClusterUnionComponent::HandleAddOrModifiedClusteredComponent(UPrimitiveCom
 
 	const bool bIsNew = !ComponentToPhysicsObjects.Contains(ChangedComponent);
 	FClusteredComponentData& ComponentData = ComponentToPhysicsObjects.FindOrAdd(ChangedComponent);
-	PendingComponentSync.Remove(ChangedComponent);
+
+	// A component shouldn't be in the PendingComponentSync map unless it's new.
+	FClusterUnionPendingAddData PendingData;
+	if (bIsNew && PendingComponentSync.RemoveAndCopyValue(ChangedComponent, PendingData))
+	{
+		ComponentData.CachedAccelerationPayloads.Append(PendingData.AccelerationPayloads);
+	}
 
 	// If this is a *new* component that we're keeping track of then there's additional book-keeping
 	// we need to do to make sure we don't forget what exactly we're tracking. Additionally, we need to
@@ -660,9 +683,11 @@ void UClusterUnionComponent::HandleAddOrModifiedClusteredComponent(UPrimitiveCom
 		{
 			if (PhysicsObject)
 			{
-				Chaos::FAccelerationStructureHandle Handle = Interface->CreateAccelerationStructureHandle(PhysicsObject);
 				const FBox HandleBounds = Interface->GetWorldBounds({ &PhysicsObject, 1 });
+				FExternalSpatialAccelerationPayload Handle;
+				Handle.Initialize(ChangedComponent, Kvp.Key);
 				AccelerationStructure->UpdateElement(Handle, Chaos::TAABB<Chaos::FReal, 3>{HandleBounds.Min, HandleBounds.Max}, HandleBounds.IsValid != 0);
+				ComponentData.CachedAccelerationPayloads.Add(Handle);
 			}
 		}
 
@@ -677,8 +702,10 @@ void UClusterUnionComponent::HandleAddOrModifiedClusteredComponent(UPrimitiveCom
 		{
 			if (Chaos::FPhysicsObjectHandle PhysicsObject = ChangedComponent->GetPhysicsObjectById(BoneId))
 			{
-				Chaos::FAccelerationStructureHandle Handle = Interface->CreateAccelerationStructureHandle(PhysicsObject);
+				FExternalSpatialAccelerationPayload Handle;
+				Handle.Initialize(ChangedComponent, BoneId);
 				AccelerationStructure->RemoveElement(Handle);
+				ComponentData.CachedAccelerationPayloads.Remove(Handle);
 			}
 		}
 	}
@@ -715,8 +742,10 @@ void UClusterUnionComponent::HandleRemovedClusteredComponent(UPrimitiveComponent
 
 				if (AccelerationStructure && PhysicsObject)
 				{
-					Chaos::FAccelerationStructureHandle Handle = Interface->CreateAccelerationStructureHandle(PhysicsObject);
+					FExternalSpatialAccelerationPayload Handle;
+					Handle.Initialize(ChangedComponent, Chaos::FPhysicsObjectInterface::GetId(PhysicsObject));
 					AccelerationStructure->RemoveElement(Handle);
+					ComponentData->CachedAccelerationPayloads.Remove(Handle);
 				}
 			}
 		}
@@ -811,7 +840,7 @@ bool UClusterUnionComponent::LineTraceComponent(TArray<FHitResult>& OutHit, cons
 	SCOPE_CYCLE_COUNTER(STAT_ClusterUnionComponentLineTraceComponentMulti);
 	if (AccelerationStructure)
 	{
-		return FGenericPhysicsInterface::RaycastMultiUsingSpatialAcceleration(*AccelerationStructure, GetWorld(), OutHit, Start, End, TraceChannel, Params, ResponseParams, ObjectParams);
+		return FGenericRaycastPhysicsInterfaceUsingSpatialAcceleration<IExternalSpatialAcceleration>::RaycastMulti(*AccelerationStructure, GetWorld(), OutHit, Start, End, TraceChannel, Params, ResponseParams, ObjectParams);
 	}
 	OutHit.Reset();
 
@@ -837,7 +866,7 @@ bool UClusterUnionComponent::SweepComponent(TArray<FHitResult>& OutHit, const FV
 	SCOPE_CYCLE_COUNTER(STAT_ClusterUnionComponentSweepComponentMulti);
 	if (AccelerationStructure)
 	{
-		return FGenericPhysicsInterface::GeomSweepMultiUsingSpatialAcceleration(*AccelerationStructure, GetWorld(), Geometry, ShapeWorldRotation, OutHit, Start, End, TraceChannel, Params, ResponseParams, ObjectParams);
+		return FGenericGeomPhysicsInterfaceUsingSpatialAcceleration<IExternalSpatialAcceleration, FPhysicsGeometry>::GeomSweepMulti(*AccelerationStructure, GetWorld(), Geometry, ShapeWorldRotation, OutHit, Start, End, TraceChannel, Params, ResponseParams, ObjectParams);
 	}
 
 	OutHit.Reset();
@@ -864,7 +893,7 @@ bool UClusterUnionComponent::LineTraceComponent(FHitResult& OutHit, const FVecto
 	SCOPE_CYCLE_COUNTER(STAT_ClusterUnionComponentLineTraceComponentSingle);
 	if (AccelerationStructure)
 	{
-		return FGenericPhysicsInterface::RaycastSingleUsingSpatialAcceleration(*AccelerationStructure, GetWorld(), OutHit, Start, End, TraceChannel, Params, ResponseParams, ObjectParams);
+		return FGenericRaycastPhysicsInterfaceUsingSpatialAcceleration<IExternalSpatialAcceleration>::RaycastSingle(*AccelerationStructure, GetWorld(), OutHit, Start, End, TraceChannel, Params, ResponseParams, ObjectParams);
 	}
 
 	bool bHasHit = false;
@@ -893,7 +922,7 @@ bool UClusterUnionComponent::SweepComponent(FHitResult& OutHit, const FVector St
 	SCOPE_CYCLE_COUNTER(STAT_ClusterUnionComponentSweepComponentSingle);
 	if (AccelerationStructure)
 	{
-		return FGenericPhysicsInterface::GeomSweepSingleUsingSpatialAcceleration(*AccelerationStructure, GetWorld(), Geometry, ShapeWorldRotation, OutHit, Start, End, TraceChannel, Params, ResponseParams, ObjectParams);
+		return FGenericGeomPhysicsInterfaceUsingSpatialAcceleration<IExternalSpatialAcceleration, FPhysicsGeometry>::GeomSweepSingle(*AccelerationStructure, GetWorld(), Geometry, ShapeWorldRotation, OutHit, Start, End, TraceChannel, Params, ResponseParams, ObjectParams);
 	}
 
 	bool bHasHit = false;
@@ -922,7 +951,7 @@ bool UClusterUnionComponent::OverlapComponentWithResult(const FVector& Pos, cons
 	SCOPE_CYCLE_COUNTER(STAT_ClusterUnionComponentOverlapComponentWithResult);
 	if (AccelerationStructure)
 	{
-		return FGenericPhysicsInterface::GeomOverlapMultiUsingSpatialAcceleration(*AccelerationStructure, GetWorld(), Geometry, Pos, Rot, OutOverlap, TraceChannel, Params, ResponseParams, ObjectParams);
+		return FGenericGeomPhysicsInterfaceUsingSpatialAcceleration<IExternalSpatialAcceleration, FPhysicsGeometry>::GeomOverlapMulti(*AccelerationStructure, GetWorld(), Geometry, Pos, Rot, OutOverlap, TraceChannel, Params, ResponseParams, ObjectParams);
 	}
 
 	bool bHasOverlap = false;
@@ -979,7 +1008,7 @@ bool UClusterUnionComponent::ComponentOverlapComponentWithResultImpl(const class
 				{
 					if (Chaos::TSerializablePtr<Chaos::FImplicitObject> Geometry = Shape->GetGeometry())
 					{
-						bHasOverlap |= FGenericPhysicsInterface::GeomOverlapMultiUsingSpatialAcceleration(*AccelerationStructure, GetWorld(), *Geometry, Pos, Rot, OutOverlap, DefaultCollisionChannel, Params, FCollisionResponseParams::DefaultResponseParam, FCollisionObjectQueryParams::DefaultObjectQueryParam);
+						bHasOverlap |= FGenericGeomPhysicsInterfaceUsingSpatialAcceleration<IExternalSpatialAcceleration, FPhysicsGeometry>::GeomOverlapMulti(*AccelerationStructure, GetWorld(), *Geometry, Pos, Rot, OutOverlap, DefaultCollisionChannel, Params, FCollisionResponseParams::DefaultResponseParam, FCollisionObjectQueryParams::DefaultObjectQueryParam);
 					}
 				}
 
diff --git a/Engine/Source/Runtime/Engine/Private/PhysicsEngine/ExternalSpatialAccelerationPayload.cpp b/Engine/Source/Runtime/Engine/Private/PhysicsEngine/ExternalSpatialAccelerationPayload.cpp
new file mode 100644
--- /dev/null
+++ b/Engine/Source/Runtime/Engine/Private/PhysicsEngine/ExternalSpatialAccelerationPayload.cpp
@@ -0,0 +1,72 @@
+// Copyright Epic Games, Inc. All Rights Reserved.
+
+#include "PhysicsEngine/ExternalSpatialAccelerationPayload.h"
+
+#include "Components/PrimitiveComponent.h"
+#include "PhysicsEngine/PhysicsObjectExternalInterface.h"
+
+#if CHAOS_DEBUG_DRAW
+#include "Chaos/ChaosDebugDraw.h"
+#endif
+
+FExternalSpatialAccelerationPayload::FExternalSpatialAccelerationPayload()
+{
+}
+
+void FExternalSpatialAccelerationPayload::Initialize(UPrimitiveComponent* InComponent, int32 InBoneId)
+{
+	Component = InComponent;
+	BoneId = InBoneId;
+
+	if (Chaos::FGeometryParticle* Particle = GetExternalGeometryParticle_ExternalThread())
+	{
+		CachedUniqueIdx = Particle->UniqueIdx();
+	}
+}
+
+Chaos::FGeometryParticle* FExternalSpatialAccelerationPayload::GetExternalGeometryParticle_ExternalThread() const
+{
+	if (UPrimitiveComponent* ValidComponent = Component.Get())
+	{
+		if (!ValidComponent->HasValidPhysicsState())
+		{
+			return nullptr;
+		}
+
+		Chaos::FPhysicsObjectHandle Handle = ValidComponent->GetPhysicsObjectById(BoneId);
+		FLockedReadPhysicsObjectExternalInterface Interface = FPhysicsObjectExternalInterface::LockRead(Handle);
+		return Interface->GetParticle(Handle);
+	}
+	
+	return nullptr;
+}
+
+Chaos::FUniqueIdx FExternalSpatialAccelerationPayload::UniqueIdx() const
+{
+	return CachedUniqueIdx;
+}
+
+void FExternalSpatialAccelerationPayload::Serialize(Chaos::FChaosArchive& Ar)
+{
+	// TODO:
+}
+
+#if CHAOS_DEBUG_DRAW
+void FExternalSpatialAccelerationPayload::DebugDraw(const bool bExternal, const bool bHit) const
+{
+	if (bExternal)
+	{
+		if (Chaos::FGeometryParticle* ExternalGeometryParticle = GetExternalGeometryParticle_ExternalThread())
+		{
+			Chaos::DebugDraw::DrawParticleShapes(Chaos::FRigidTransform3(), ExternalGeometryParticle, bHit ? FColor::Red : FColor::Green);
+		}
+	}
+}
+#endif
+
+uint32 GetTypeHash(const FExternalSpatialAccelerationPayload& Payload)
+{
+	const uint32 ComponentHash = Payload.Component.GetWeakPtrTypeHash();
+	const uint32 BoneIdHash = GetTypeHash(Payload.BoneId);
+	return HashCombine(ComponentHash, BoneIdHash);
+}
\ No newline at end of file
diff --git a/Engine/Source/Runtime/Engine/Public/Physics/Experimental/ChaosInterfaceWrapper.h b/Engine/Source/Runtime/Engine/Public/Physics/Experimental/ChaosInterfaceWrapper.h
--- a/Engine/Source/Runtime/Engine/Public/Physics/Experimental/ChaosInterfaceWrapper.h
+++ b/Engine/Source/Runtime/Engine/Public/Physics/Experimental/ChaosInterfaceWrapper.h
@@ -42,17 +42,11 @@ UPhysicalMaterial* GetUserData(const Chaos::FChaosPhysicsMaterial& Material);
 
 }
 
-void LowLevelRaycast(const FPhysScene& Scene, const FVector& Start, const FVector& Dir, float DeltaMag, FPhysicsHitCallback<FHitRaycast>& HitBuffer, EHitFlags OutputFlags, FQueryFlags QueryFlags, const FCollisionFilterData& Filter, const ChaosInterface::FQueryFilterData& QueryFilterData, ICollisionQueryFilterCallbackBase* QueryCallback, const ChaosInterface::FQueryDebugParams& DebugParams = {});
-void LowLevelRaycast(const FPhysScene& Scene, const FVector& Start, const FVector& Dir, float DeltaMag, FPhysicsHitCallback<ChaosInterface::FPTRaycastHit>& HitBuffer, EHitFlags OutputFlags, FQueryFlags QueryFlags, const FCollisionFilterData& Filter, const ChaosInterface::FQueryFilterData& QueryFilterData, ICollisionQueryFilterCallbackBase* QueryCallback, const ChaosInterface::FQueryDebugParams& DebugParams = {});
-void LowLevelRaycast(const Chaos::IDefaultChaosSpatialAcceleration& AccelerationStructure, const FVector& Start, const FVector& Dir, float DeltaMag, FPhysicsHitCallback<FHitRaycast>& HitBuffer, EHitFlags OutputFlags, FQueryFlags QueryFlags, const FCollisionFilterData& Filter, const ChaosInterface::FQueryFilterData& QueryFilterData, ICollisionQueryFilterCallbackBase* QueryCallback, const ChaosInterface::FQueryDebugParams& DebugParams = {});
-void LowLevelRaycast(const Chaos::IDefaultChaosSpatialAcceleration& AccelerationStructure, const FVector& Start, const FVector& Dir, float DeltaMag, FPhysicsHitCallback<ChaosInterface::FPTRaycastHit>& HitBuffer, EHitFlags OutputFlags, FQueryFlags QueryFlags, const FCollisionFilterData& Filter, const ChaosInterface::FQueryFilterData& QueryFilterData, ICollisionQueryFilterCallbackBase* QueryCallback, const ChaosInterface::FQueryDebugParams& DebugParams = {});
-
-void LowLevelSweep(const FPhysScene& Scene, const FPhysicsGeometry& Geom, const FTransform& StartTM, const FVector& Dir, float DeltaMag, FPhysicsHitCallback<FHitSweep>& HitBuffer, EHitFlags OutputFlags, FQueryFlags QueryFlags, const FCollisionFilterData& Filter, const ChaosInterface::FQueryFilterData& QueryFilterData, ICollisionQueryFilterCallbackBase* QueryCallback, const ChaosInterface::FQueryDebugParams& DebugParams = {});
-void LowLevelSweep(const FPhysScene& Scene, const FPhysicsGeometry& Geom, const FTransform& StartTM, const FVector& Dir, float DeltaMag, FPhysicsHitCallback<ChaosInterface::FPTSweepHit>& HitBuffer, EHitFlags OutputFlags, FQueryFlags QueryFlags, const FCollisionFilterData& Filter, const ChaosInterface::FQueryFilterData& QueryFilterData, ICollisionQueryFilterCallbackBase* QueryCallback, const ChaosInterface::FQueryDebugParams& DebugParams = {});
-void LowLevelSweep(const Chaos::IDefaultChaosSpatialAcceleration& AccelerationStructure, const FPhysicsGeometry& Geom, const FTransform& StartTM, const FVector& Dir, float DeltaMag, FPhysicsHitCallback<FHitSweep>& HitBuffer, EHitFlags OutputFlags, FQueryFlags QueryFlags, const FCollisionFilterData& Filter, const ChaosInterface::FQueryFilterData& QueryFilterData, ICollisionQueryFilterCallbackBase* QueryCallback, const ChaosInterface::FQueryDebugParams& DebugParams = {});
-void LowLevelSweep(const Chaos::IDefaultChaosSpatialAcceleration& AccelerationStructure, const FPhysicsGeometry& Geom, const FTransform& StartTM, const FVector& Dir, float DeltaMag, FPhysicsHitCallback<ChaosInterface::FPTSweepHit>& HitBuffer, EHitFlags OutputFlags, FQueryFlags QueryFlags, const FCollisionFilterData& Filter, const ChaosInterface::FQueryFilterData& QueryFilterData, ICollisionQueryFilterCallbackBase* QueryCallback, const ChaosInterface::FQueryDebugParams& DebugParams = {});
-
-void LowLevelOverlap(const FPhysScene& Scene, const FPhysicsGeometry& Geom, const FTransform& GeomPose, FPhysicsHitCallback<FHitOverlap>& HitBuffer, FQueryFlags QueryFlags, const FCollisionFilterData& Filter, const ChaosInterface::FQueryFilterData& QueryFilterData, ICollisionQueryFilterCallbackBase* QueryCallback, const ChaosInterface::FQueryDebugParams& DebugParams = {});
-void LowLevelOverlap(const FPhysScene& Scene, const FPhysicsGeometry& Geom, const FTransform& GeomPose, FPhysicsHitCallback<ChaosInterface::FPTOverlapHit>& HitBuffer, FQueryFlags QueryFlags, const FCollisionFilterData& Filter, const ChaosInterface::FQueryFilterData& QueryFilterData, ICollisionQueryFilterCallbackBase* QueryCallback, const ChaosInterface::FQueryDebugParams& DebugParams = {});
-void LowLevelOverlap(const Chaos::IDefaultChaosSpatialAcceleration& AccelerationStructure, const FPhysicsGeometry& Geom, const FTransform& GeomPose, FPhysicsHitCallback<FHitOverlap>& HitBuffer, FQueryFlags QueryFlags, const FCollisionFilterData& Filter, const ChaosInterface::FQueryFilterData& QueryFilterData, ICollisionQueryFilterCallbackBase* QueryCallback, const ChaosInterface::FQueryDebugParams& DebugParams = {});
-void LowLevelOverlap(const Chaos::IDefaultChaosSpatialAcceleration& AccelerationStructure, const FPhysicsGeometry& Geom, const FTransform& GeomPose, FPhysicsHitCallback<ChaosInterface::FPTOverlapHit>& HitBuffer, FQueryFlags QueryFlags, const FCollisionFilterData& Filter, const ChaosInterface::FQueryFilterData& QueryFilterData, ICollisionQueryFilterCallbackBase* QueryCallback, const ChaosInterface::FQueryDebugParams& DebugParams = {});
\ No newline at end of file
+template<typename TContainer, typename THit>
+void LowLevelRaycast(const TContainer& Container, const FVector& Start, const FVector& Dir, float DeltaMag, FPhysicsHitCallback<THit>& HitBuffer, EHitFlags OutputFlags, FQueryFlags QueryFlags, const FCollisionFilterData& Filter, const ChaosInterface::FQueryFilterData& QueryFilterData, ICollisionQueryFilterCallbackBase* QueryCallback, const ChaosInterface::FQueryDebugParams& DebugParams = {});
+
+template<typename TContainer, typename THit>
+void LowLevelSweep(const TContainer& Container, const FPhysicsGeometry& Geom, const FTransform& StartTM, const FVector& Dir, float DeltaMag, FPhysicsHitCallback<THit>& HitBuffer, EHitFlags OutputFlags, FQueryFlags QueryFlags, const FCollisionFilterData& Filter, const ChaosInterface::FQueryFilterData& QueryFilterData, ICollisionQueryFilterCallbackBase* QueryCallback, const ChaosInterface::FQueryDebugParams& DebugParams = {});
+
+template<typename TContainer, typename THit>
+void LowLevelOverlap(const TContainer& Container, const FPhysicsGeometry& Geom, const FTransform& GeomPose, FPhysicsHitCallback<THit>& HitBuffer, FQueryFlags QueryFlags, const FCollisionFilterData& Filter, const ChaosInterface::FQueryFilterData& QueryFilterData, ICollisionQueryFilterCallbackBase* QueryCallback, const ChaosInterface::FQueryDebugParams& DebugParams = {});
\ No newline at end of file
diff --git a/Engine/Source/Runtime/Engine/Public/Physics/GenericPhysicsInterface.h b/Engine/Source/Runtime/Engine/Public/Physics/GenericPhysicsInterface.h
--- a/Engine/Source/Runtime/Engine/Public/Physics/GenericPhysicsInterface.h
+++ b/Engine/Source/Runtime/Engine/Public/Physics/GenericPhysicsInterface.h
@@ -15,11 +15,9 @@ struct ENGINE_API FGenericPhysicsInterface
 {
 	/** Trace a ray against the world and return if a blocking hit is found */
 	static bool RaycastTest(const UWorld* World, const FVector Start, const FVector End, ECollisionChannel TraceChannel, const FCollisionQueryParams& Params, const FCollisionResponseParams& ResponseParams, const FCollisionObjectQueryParams& ObjectParams = FCollisionObjectQueryParams::DefaultObjectQueryParam);
-	static bool RaycastTestUsingSpatialAcceleration(const Chaos::IDefaultChaosSpatialAcceleration& Accel, const UWorld* World, const FVector Start, const FVector End, ECollisionChannel TraceChannel, const FCollisionQueryParams& Params, const FCollisionResponseParams& ResponseParams, const FCollisionObjectQueryParams& ObjectParams = FCollisionObjectQueryParams::DefaultObjectQueryParam);
 
 	/** Trace a ray against the world and return the first blocking hit */
 	static bool RaycastSingle(const UWorld* World, struct FHitResult& OutHit, const FVector Start, const FVector End, ECollisionChannel TraceChannel, const FCollisionQueryParams& Params, const FCollisionResponseParams& ResponseParams, const FCollisionObjectQueryParams& ObjectParams = FCollisionObjectQueryParams::DefaultObjectQueryParam);
-	static bool RaycastSingleUsingSpatialAcceleration(const Chaos::IDefaultChaosSpatialAcceleration& Accel, const UWorld* World, struct FHitResult& OutHit, const FVector Start, const FVector End, ECollisionChannel TraceChannel, const FCollisionQueryParams& Params, const FCollisionResponseParams& ResponseParams, const FCollisionObjectQueryParams& ObjectParams = FCollisionObjectQueryParams::DefaultObjectQueryParam);
 
 	/**
 	*  Trace a ray against the world and return touching hits and then first blocking hit
@@ -27,40 +25,46 @@ struct ENGINE_API FGenericPhysicsInterface
 	*  Only the single closest blocking result will be generated, no tests will be done after that
 	*/
 	static bool RaycastMulti(const UWorld* World, TArray<struct FHitResult>& OutHits, const FVector& Start, const FVector& End, ECollisionChannel TraceChannel, const FCollisionQueryParams& Params, const FCollisionResponseParams& ResponseParams, const FCollisionObjectQueryParams& ObjectParams = FCollisionObjectQueryParams::DefaultObjectQueryParam);
-	static bool RaycastMultiUsingSpatialAcceleration(const Chaos::IDefaultChaosSpatialAcceleration& Accel, const UWorld* World, TArray<struct FHitResult>& OutHits, const FVector& Start, const FVector& End, ECollisionChannel TraceChannel, const FCollisionQueryParams& Params, const FCollisionResponseParams& ResponseParams, const FCollisionObjectQueryParams& ObjectParams = FCollisionObjectQueryParams::DefaultObjectQueryParam);
 
 	/** Function used for sweeping a supplied shape against the world as a test */
 	static bool GeomSweepTest(const UWorld* World, const FCollisionShape& CollisionShape, const FQuat& Rot, FVector Start, FVector End, ECollisionChannel TraceChannel, const FCollisionQueryParams& Params, const FCollisionResponseParams& ResponseParams, const FCollisionObjectQueryParams& ObjectParams = FCollisionObjectQueryParams::DefaultObjectQueryParam);
-	static bool GeomSweepTestUsingSpatialAcceleration(const Chaos::IDefaultChaosSpatialAcceleration& Accel, const UWorld* World, const FCollisionShape& CollisionShape, const FQuat& Rot, FVector Start, FVector End, ECollisionChannel TraceChannel, const FCollisionQueryParams& Params, const FCollisionResponseParams& ResponseParams, const FCollisionObjectQueryParams& ObjectParams = FCollisionObjectQueryParams::DefaultObjectQueryParam);
 
 	/** Function for sweeping a supplied shape against the world */
 	template<typename GeomWrapper>
 	static bool GeomSweepSingle(const UWorld* World, const GeomWrapper& InGeom, const FQuat& Rot, FHitResult& OutHit, FVector Start, FVector End, ECollisionChannel TraceChannel, const FCollisionQueryParams& Params, const FCollisionResponseParams& ResponseParams, const FCollisionObjectQueryParams& ObjectParams = FCollisionObjectQueryParams::DefaultObjectQueryParam);
 
-	template<typename GeomWrapper>
-	static bool GeomSweepSingleUsingSpatialAcceleration(const Chaos::IDefaultChaosSpatialAcceleration& Accel, const UWorld* World, const GeomWrapper& InGeom, const FQuat& Rot, FHitResult& OutHit, FVector Start, FVector End, ECollisionChannel TraceChannel, const FCollisionQueryParams& Params, const FCollisionResponseParams& ResponseParams, const FCollisionObjectQueryParams& ObjectParams = FCollisionObjectQueryParams::DefaultObjectQueryParam);
-
 	/** Sweep a supplied shape against the world and do not stop until the first blocking hit */
 	template <typename GeomWrapper>
 	static bool GeomSweepMulti(const UWorld* World, const GeomWrapper& InGeom, const FQuat& InGeomRot, TArray<FHitResult>& OutHits, FVector Start, FVector End, ECollisionChannel TraceChannel, const FCollisionQueryParams& Params, const FCollisionResponseParams& ResponseParams, const FCollisionObjectQueryParams& ObjectParams = FCollisionObjectQueryParams::DefaultObjectQueryParam);
 
-	template<typename GeomWrapper>
-	static bool GeomSweepMultiUsingSpatialAcceleration(const Chaos::IDefaultChaosSpatialAcceleration& Accel, const UWorld* World, const GeomWrapper& InGeom, const FQuat& InGeomRot, TArray<FHitResult>& OutHits, FVector Start, FVector End, ECollisionChannel TraceChannel, const FCollisionQueryParams& Params, const FCollisionResponseParams& ResponseParams, const FCollisionObjectQueryParams& ObjectParams = FCollisionObjectQueryParams::DefaultObjectQueryParam);
-
 	/** Find overlapping shapes with a given shape */
 	template<typename GeomWrapper>
 	static bool GeomOverlapMulti(const UWorld* World, const GeomWrapper& InGeom, const FVector& InPosition, const FQuat& InRotation, TArray<FOverlapResult>& OutOverlaps, ECollisionChannel TraceChannel, const FCollisionQueryParams& Params, const FCollisionResponseParams& ResponseParams, const FCollisionObjectQueryParams& ObjectParams);
 
-	template<typename GeomWrapper>
-	static bool GeomOverlapMultiUsingSpatialAcceleration(const Chaos::IDefaultChaosSpatialAcceleration& Accel, const UWorld* World, const GeomWrapper& InGeom, const FVector& InPosition, const FQuat& InRotation, TArray<FOverlapResult>& OutOverlaps, ECollisionChannel TraceChannel, const FCollisionQueryParams& Params, const FCollisionResponseParams& ResponseParams, const FCollisionObjectQueryParams& ObjectParams);
-
 	/** Function for testing overlaps between a supplied PxGeometry and the world. Returns true if at least one overlapping shape is blocking*/
 	static bool GeomOverlapBlockingTest(const UWorld* World, const FCollisionShape& CollisionShape, const FVector& Pos, const FQuat& Rot, ECollisionChannel TraceChannel, const FCollisionQueryParams& Params, const FCollisionResponseParams& ResponseParams, const FCollisionObjectQueryParams& ObjectParams = FCollisionObjectQueryParams::DefaultObjectQueryParam);
-	static bool GeomOverlapBlockingTestUsingSpatialAcceleration(const Chaos::IDefaultChaosSpatialAcceleration& Accel, const UWorld* World, const FCollisionShape& CollisionShape, const FVector& Pos, const FQuat& Rot, ECollisionChannel TraceChannel, const FCollisionQueryParams& Params, const FCollisionResponseParams& ResponseParams, const FCollisionObjectQueryParams& ObjectParams = FCollisionObjectQueryParams::DefaultObjectQueryParam);
 
 	/** Function for testing overlaps between a supplied PxGeometry and the world. Returns true if anything is overlapping (blocking or touching)*/
 	static bool GeomOverlapAnyTest(const UWorld* World, const FCollisionShape& CollisionShape, const FVector& Pos, const FQuat& Rot, ECollisionChannel TraceChannel, const FCollisionQueryParams& Params, const FCollisionResponseParams& ResponseParams, const FCollisionObjectQueryParams& ObjectParams = FCollisionObjectQueryParams::DefaultObjectQueryParam);
-	static bool GeomOverlapAnyTestUsingSpatialAcceleration(const Chaos::IDefaultChaosSpatialAcceleration& Accel, const UWorld* World, const FCollisionShape& CollisionShape, const FVector& Pos, const FQuat& Rot, ECollisionChannel TraceChannel, const FCollisionQueryParams& Params, const FCollisionResponseParams& ResponseParams, const FCollisionObjectQueryParams& ObjectParams = FCollisionObjectQueryParams::DefaultObjectQueryParam);
+};
+
+template<typename TAccel>
+struct ENGINE_API FGenericRaycastPhysicsInterfaceUsingSpatialAcceleration
+{
+	static bool RaycastTest(const TAccel& Accel, const UWorld* World, const FVector Start, const FVector End, ECollisionChannel TraceChannel, const FCollisionQueryParams& Params, const FCollisionResponseParams& ResponseParams, const FCollisionObjectQueryParams& ObjectParams = FCollisionObjectQueryParams::DefaultObjectQueryParam);
+	static bool RaycastSingle(const TAccel& Accel, const UWorld* World, struct FHitResult& OutHit, const FVector Start, const FVector End, ECollisionChannel TraceChannel, const FCollisionQueryParams& Params, const FCollisionResponseParams& ResponseParams, const FCollisionObjectQueryParams& ObjectParams = FCollisionObjectQueryParams::DefaultObjectQueryParam);
+	static bool RaycastMulti(const TAccel& Accel, const UWorld* World, TArray<struct FHitResult>& OutHits, const FVector& Start, const FVector& End, ECollisionChannel TraceChannel, const FCollisionQueryParams& Params, const FCollisionResponseParams& ResponseParams, const FCollisionObjectQueryParams& ObjectParams = FCollisionObjectQueryParams::DefaultObjectQueryParam);
+};
+
+template<typename TAccel, typename TGeom>
+struct ENGINE_API FGenericGeomPhysicsInterfaceUsingSpatialAcceleration
+{
+	static bool GeomSweepTest(const TAccel& Accel, const UWorld* World, const TGeom& InGeom, const FQuat& Rot, FVector Start, FVector End, ECollisionChannel TraceChannel, const FCollisionQueryParams& Params, const FCollisionResponseParams& ResponseParams, const FCollisionObjectQueryParams& ObjectParams = FCollisionObjectQueryParams::DefaultObjectQueryParam);
+	static bool GeomSweepSingle(const TAccel& Accel, const UWorld* World, const TGeom& InGeom, const FQuat& Rot, FHitResult& OutHit, FVector Start, FVector End, ECollisionChannel TraceChannel, const FCollisionQueryParams& Params, const FCollisionResponseParams& ResponseParams, const FCollisionObjectQueryParams& ObjectParams = FCollisionObjectQueryParams::DefaultObjectQueryParam);
+	static bool GeomSweepMulti(const TAccel& Accel, const UWorld* World, const TGeom& InGeom, const FQuat& InGeomRot, TArray<FHitResult>& OutHits, FVector Start, FVector End, ECollisionChannel TraceChannel, const FCollisionQueryParams& Params, const FCollisionResponseParams& ResponseParams, const FCollisionObjectQueryParams& ObjectParams = FCollisionObjectQueryParams::DefaultObjectQueryParam);
+	static bool GeomOverlapMulti(const TAccel& Accel, const UWorld* World, const TGeom& InGeom, const FVector& InPosition, const FQuat& InRotation, TArray<FOverlapResult>& OutOverlaps, ECollisionChannel TraceChannel, const FCollisionQueryParams& Params, const FCollisionResponseParams& ResponseParams, const FCollisionObjectQueryParams& ObjectParams);
+	static bool GeomOverlapBlockingTest(const TAccel& Accel, const UWorld* World, const TGeom& InGeom, const FVector& Pos, const FQuat& Rot, ECollisionChannel TraceChannel, const FCollisionQueryParams& Params, const FCollisionResponseParams& ResponseParams, const FCollisionObjectQueryParams& ObjectParams = FCollisionObjectQueryParams::DefaultObjectQueryParam);
+	static bool GeomOverlapAnyTest(const TAccel& Accel, const UWorld* World, const TGeom& InGeom, const FVector& Pos, const FQuat& Rot, ECollisionChannel TraceChannel, const FCollisionQueryParams& Params, const FCollisionResponseParams& ResponseParams, const FCollisionObjectQueryParams& ObjectParams = FCollisionObjectQueryParams::DefaultObjectQueryParam);
 };
 
 template<>
@@ -88,33 +92,4 @@ template<>
 ENGINE_API bool FGenericPhysicsInterface::GeomSweepMulti(const UWorld* World, const FPhysicsGeometryCollection& InGeom, const FQuat& InGeomRot, TArray<FHitResult>& OutHits, FVector Start, FVector End, ECollisionChannel TraceChannel, const FCollisionQueryParams& Params, const FCollisionResponseParams& ResponseParams, const FCollisionObjectQueryParams& ObjectParams);
 
 template<>
-ENGINE_API bool FGenericPhysicsInterface::GeomOverlapMulti(const UWorld* World, const FPhysicsGeometryCollection& InGeom, const FVector& InPosition, const FQuat& InRotation, TArray<FOverlapResult>& OutOverlaps, ECollisionChannel TraceChannel, const FCollisionQueryParams& Params, const FCollisionResponseParams& ResponseParams, const FCollisionObjectQueryParams& ObjectParams);
-
-////////////////////////////////////
-
-template<>
-ENGINE_API bool FGenericPhysicsInterface::GeomSweepSingleUsingSpatialAcceleration(const Chaos::IDefaultChaosSpatialAcceleration& Accel, const UWorld* World, const FCollisionShape& InGeom, const FQuat& Rot, FHitResult& OutHit, FVector Start, FVector End, ECollisionChannel TraceChannel, const FCollisionQueryParams& Params, const FCollisionResponseParams& ResponseParams, const FCollisionObjectQueryParams& ObjectParams);
-
-template<>
-ENGINE_API bool FGenericPhysicsInterface::GeomSweepSingleUsingSpatialAcceleration(const Chaos::IDefaultChaosSpatialAcceleration& Accel, const UWorld* World, const FPhysicsGeometry& InGeom, const FQuat& Rot, FHitResult& OutHit, FVector Start, FVector End, ECollisionChannel TraceChannel, const FCollisionQueryParams& Params, const FCollisionResponseParams& ResponseParams, const FCollisionObjectQueryParams& ObjectParams);
-
-template<>
-ENGINE_API bool FGenericPhysicsInterface::GeomSweepMultiUsingSpatialAcceleration(const Chaos::IDefaultChaosSpatialAcceleration& Accel, const UWorld* World, const FCollisionShape& InGeom, const FQuat& InGeomRot, TArray<FHitResult>& OutHits, FVector Start, FVector End, ECollisionChannel TraceChannel, const FCollisionQueryParams& Params, const FCollisionResponseParams& ResponseParams, const FCollisionObjectQueryParams& ObjectParams);
-
-template<>
-ENGINE_API bool FGenericPhysicsInterface::GeomSweepMultiUsingSpatialAcceleration(const Chaos::IDefaultChaosSpatialAcceleration& Accel, const UWorld* World, const FPhysicsGeometry& InGeom, const FQuat& InGeomRot, TArray<FHitResult>& OutHits, FVector Start, FVector End, ECollisionChannel TraceChannel, const FCollisionQueryParams& Params, const FCollisionResponseParams& ResponseParams, const FCollisionObjectQueryParams& ObjectParams);
-
-template<>
-ENGINE_API bool FGenericPhysicsInterface::GeomOverlapMultiUsingSpatialAcceleration(const Chaos::IDefaultChaosSpatialAcceleration& Accel, const UWorld* World, const FCollisionShape& InGeom, const FVector& InPosition, const FQuat& InRotation, TArray<FOverlapResult>& OutOverlaps, ECollisionChannel TraceChannel, const FCollisionQueryParams& Params, const FCollisionResponseParams& ResponseParams, const FCollisionObjectQueryParams& ObjectParams);
-
-template<>
-ENGINE_API bool FGenericPhysicsInterface::GeomOverlapMultiUsingSpatialAcceleration(const Chaos::IDefaultChaosSpatialAcceleration& Accel, const UWorld* World, const FPhysicsGeometry& InGeom, const FVector& InPosition, const FQuat& InRotation, TArray<FOverlapResult>& OutOverlaps, ECollisionChannel TraceChannel, const FCollisionQueryParams& Params, const FCollisionResponseParams& ResponseParams, const FCollisionObjectQueryParams& ObjectParams);
-
-template<>
-ENGINE_API bool FGenericPhysicsInterface::GeomSweepSingleUsingSpatialAcceleration(const Chaos::IDefaultChaosSpatialAcceleration& Accel, const UWorld* World, const FPhysicsGeometryCollection& InGeom, const FQuat& Rot, FHitResult& OutHit, FVector Start, FVector End, ECollisionChannel TraceChannel, const FCollisionQueryParams& Params, const FCollisionResponseParams& ResponseParams, const FCollisionObjectQueryParams& ObjectParams);
-
-template<>
-ENGINE_API bool FGenericPhysicsInterface::GeomSweepMultiUsingSpatialAcceleration(const Chaos::IDefaultChaosSpatialAcceleration& Accel, const UWorld* World, const FPhysicsGeometryCollection& InGeom, const FQuat& InGeomRot, TArray<FHitResult>& OutHits, FVector Start, FVector End, ECollisionChannel TraceChannel, const FCollisionQueryParams& Params, const FCollisionResponseParams& ResponseParams, const FCollisionObjectQueryParams& ObjectParams);
-
-template<>
-ENGINE_API bool FGenericPhysicsInterface::GeomOverlapMultiUsingSpatialAcceleration(const Chaos::IDefaultChaosSpatialAcceleration& Accel, const UWorld* World, const FPhysicsGeometryCollection& InGeom, const FVector& InPosition, const FQuat& InRotation, TArray<FOverlapResult>& OutOverlaps, ECollisionChannel TraceChannel, const FCollisionQueryParams& Params, const FCollisionResponseParams& ResponseParams, const FCollisionObjectQueryParams& ObjectParams);
\ No newline at end of file
+ENGINE_API bool FGenericPhysicsInterface::GeomOverlapMulti(const UWorld* World, const FPhysicsGeometryCollection& InGeom, const FVector& InPosition, const FQuat& InRotation, TArray<FOverlapResult>& OutOverlaps, ECollisionChannel TraceChannel, const FCollisionQueryParams& Params, const FCollisionResponseParams& ResponseParams, const FCollisionObjectQueryParams& ObjectParams);
\ No newline at end of file
diff --git a/Engine/Source/Runtime/Experimental/Chaos/Private/Chaos/BoundingVolume.cpp b/Engine/Source/Runtime/Experimental/Chaos/Private/Chaos/BoundingVolume.cpp
--- a/Engine/Source/Runtime/Experimental/Chaos/Private/Chaos/BoundingVolume.cpp
+++ b/Engine/Source/Runtime/Experimental/Chaos/Private/Chaos/BoundingVolume.cpp
@@ -26,29 +26,6 @@ namespace Chaos
 	    MaxDirtyElements,
 	    TEXT("The max number of dirty elements. This forces a flush which is very expensive"));
 
-	    template<typename TPayloadType, typename T, int d>
-	    ISpatialAcceleration<TPayloadType, T, d>* ISpatialAcceleration<TPayloadType, T, d>::SerializationFactory(FChaosArchive& Ar, ISpatialAcceleration<TPayloadType, T, d>* Accel)
-	{
-		if (Ar.CustomVer(FExternalPhysicsCustomObjectVersion::GUID) < FExternalPhysicsCustomObjectVersion::SerializeEvolutionGenericAcceleration)
-		{
-			return new TBoundingVolume<TPayloadType, T, d>();
-		}
-
-		int8 AccelType = Ar.IsLoading() ? 0 : (int8)Accel->Type;
-		Ar << AccelType;
-		switch ((ESpatialAcceleration)AccelType)
-		{
-		case ESpatialAcceleration::BoundingVolume: return Ar.IsLoading() ? new TBoundingVolume<TPayloadType, T, d>() : nullptr;
-		case ESpatialAcceleration::AABBTree: return Ar.IsLoading() ? new TAABBTree<TPayloadType, TAABBTreeLeafArray<TPayloadType>>() : nullptr;
-		case ESpatialAcceleration::AABBTreeBV: return Ar.IsLoading() ? new TAABBTree<TPayloadType, TBoundingVolume<TPayloadType, T, 3>>() : nullptr;
-		case ESpatialAcceleration::Collection: check(false);	//Collections must be serialized directly since they are variadic
-		default: check(false); return nullptr;
-		}
-	}
-
-	template class CHAOS_API Chaos::ISpatialAcceleration<int32, FReal, 3>;
-	template class CHAOS_API Chaos::ISpatialAcceleration<FAccelerationStructureHandle, FReal, 3>;
-
 	template class CHAOS_API Chaos::TBoundingVolume<int32, FReal, 3>;
 	template class CHAOS_API Chaos::TBoundingVolume<FAccelerationStructureHandle, FReal, 3>;
 }
diff --git a/Engine/Source/Runtime/Experimental/Chaos/Private/Chaos/SpatialAcceleration.cpp b/Engine/Source/Runtime/Experimental/Chaos/Private/Chaos/SpatialAcceleration.cpp
new file mode 100644
--- /dev/null
+++ b/Engine/Source/Runtime/Experimental/Chaos/Private/Chaos/SpatialAcceleration.cpp
@@ -0,0 +1,34 @@
+// Copyright Epic Games, Inc. All Rights Reserved.
+
+#include "Chaos/ISpatialAcceleration.h"
+
+#include "Chaos/AABBTree.h"
+#include "Chaos/BoundingVolume.h"
+#include "Chaos/ParticleHandle.h"
+
+namespace Chaos
+{
+	template<ESpatialAcceleration SpatialType, typename TPayloadType, typename T, int d>
+	ISpatialAcceleration<TPayloadType, T, d>* TSpatialAccelerationSerializationFactory<SpatialType, TPayloadType, T, d>::Create()
+	{
+		if constexpr (SpatialType == ESpatialAcceleration::BoundingVolume)
+		{
+			return new TBoundingVolume<TPayloadType, T, d>();
+		}
+		else if constexpr (SpatialType == ESpatialAcceleration::AABBTree)
+		{
+			return new TAABBTree<TPayloadType, TAABBTreeLeafArray<TPayloadType>>();
+		}
+		else if constexpr (SpatialType == ESpatialAcceleration::AABBTreeBV)
+		{
+			return new TAABBTree<TPayloadType, TBoundingVolume<TPayloadType, T, 3>>();
+		}
+		return nullptr;
+	}
+
+	template class CHAOS_API Chaos::ISpatialAcceleration<int32, FReal, 3>;
+	template class CHAOS_API Chaos::ISpatialAcceleration<FAccelerationStructureHandle, FReal, 3>;
+
+	template class CHAOS_API Chaos::ISpatialVisitor<int32, FReal>;
+	template class CHAOS_API Chaos::ISpatialVisitor<FAccelerationStructureHandle, FReal>;
+}
\ No newline at end of file
diff --git a/Engine/Source/Runtime/Experimental/Chaos/Public/Chaos/AABBTree.h b/Engine/Source/Runtime/Experimental/Chaos/Public/Chaos/AABBTree.h
--- a/Engine/Source/Runtime/Experimental/Chaos/Public/Chaos/AABBTree.h
+++ b/Engine/Source/Runtime/Experimental/Chaos/Public/Chaos/AABBTree.h
@@ -3001,14 +3001,21 @@ private:
 
 		ensure(WorkStack.Num() == 0);
 
-		const int32 ExpectedNumLeaves = Particles.Num() / MaxChildrenInLeaf;
+		int32 NumParticles = 0;
+		constexpr bool bIsValidParticleArray = !std::is_same_v<TParticles, std::nullptr_t>;
+		if constexpr (bIsValidParticleArray)
+		{
+			NumParticles = Particles.Num();
+		}
+
+		const int32 ExpectedNumLeaves = NumParticles / MaxChildrenInLeaf;
 		const int32 ExpectedNumNodes = ExpectedNumLeaves;
 
 		WorkStack.Reserve(ExpectedNumNodes);
 
 		const int32 CurIdx = GetNewWorkSnapshot();
 		FWorkSnapshot& WorkSnapshot = WorkPool[CurIdx];
-		WorkSnapshot.Elems.Reserve(Particles.Num());
+		WorkSnapshot.Elems.Reserve(NumParticles);
 		
 		
 		GlobalPayloads.Reset();
@@ -3035,14 +3042,17 @@ private:
 		if (bDynamicTree)
 		{
 			int32 Idx = 0;
-			for (auto& Particle : Particles)
+			if constexpr (bIsValidParticleArray)
 			{
-				bool bHasBoundingBox = HasBoundingBox(Particle);
-				auto Payload = Particle.template GetPayload<TPayloadType>(Idx);
-				TAABB<T, 3> ElemBounds = ComputeWorldSpaceBoundingBox(Particle, false, (T)0);
-				
-				UpdateElement(Payload, ElemBounds, bHasBoundingBox);
-				++Idx;
+				for (auto& Particle : Particles)
+				{
+					bool bHasBoundingBox = HasBoundingBox(Particle);
+					auto Payload = Particle.template GetPayload<TPayloadType>(Idx);
+					TAABB<T, 3> ElemBounds = ComputeWorldSpaceBoundingBox(Particle, false, (T)0);
+
+					UpdateElement(Payload, ElemBounds, bHasBoundingBox);
+					++Idx;
+				}
 			}
 			this->SetAsyncTimeSlicingComplete(true);
 			return;
@@ -3062,58 +3072,61 @@ private:
 			//TODO: we need a better way to time-slice this case since there can be a huge number of Particles. Can't do it right now without making full copy
 			TVec3<T> CenterSum(0);
 
-			for (auto& Particle : Particles)
+			if constexpr (bIsValidParticleArray)
 			{
-				bool bHasBoundingBox = HasBoundingBox(Particle);
-				auto Payload = Particle.template GetPayload<TPayloadType>(Idx);
-				TAABB<T, 3> ElemBounds = ComputeWorldSpaceBoundingBox(Particle, false, (T)0);
-
-				// If bounds are bad, use global so we won't screw up splitting computations.
-				if (bHasBoundingBox && ValidateBounds(ElemBounds) == false)
+				for (auto& Particle : Particles)
 				{
-					bHasBoundingBox = false;
-					ensureMsgf(false, TEXT("AABBTree encountered invalid bounds input. Forcing element to global payload. Min: %s Max: %s."),
-						*ElemBounds.Min().ToString(), *ElemBounds.Max().ToString());
-				}
-
+					bool bHasBoundingBox = HasBoundingBox(Particle);
+					auto Payload = Particle.template GetPayload<TPayloadType>(Idx);
+					TAABB<T, 3> ElemBounds = ComputeWorldSpaceBoundingBox(Particle, false, (T)0);
 
-				if (bHasBoundingBox)
-				{
-					if (ElemBounds.Extents().Max() > MaxPayloadBounds)
+					// If bounds are bad, use global so we won't screw up splitting computations.
+					if (bHasBoundingBox && ValidateBounds(ElemBounds) == false)
 					{
 						bHasBoundingBox = false;
+						ensureMsgf(false, TEXT("AABBTree encountered invalid bounds input. Forcing element to global payload. Min: %s Max: %s."),
+							*ElemBounds.Min().ToString(), *ElemBounds.Max().ToString());
 					}
-					else
+
+
+					if (bHasBoundingBox)
 					{
-						FReal NumElems = (FReal)(WorkSnapshot.Elems.Add(FElement{ Payload, ElemBounds }) + 1);
-						WorkSnapshot.Bounds.GrowToInclude(ElemBounds);
+						if (ElemBounds.Extents().Max() > MaxPayloadBounds)
+						{
+							bHasBoundingBox = false;
+						}
+						else
+						{
+							FReal NumElems = (FReal)(WorkSnapshot.Elems.Add(FElement{ Payload, ElemBounds }) + 1);
+							WorkSnapshot.Bounds.GrowToInclude(ElemBounds);
 
-						// Include the current particle in the average and scaled variance of the particle centers using Welford's method.
-						TVec3<T> CenterDelta = ElemBounds.Center() - WorkSnapshot.AverageCenter;
-						WorkSnapshot.AverageCenter += CenterDelta / (T)NumElems;
-						WorkSnapshot.ScaledCenterVariance += (ElemBounds.Center() - WorkSnapshot.AverageCenter) * CenterDelta;
+							// Include the current particle in the average and scaled variance of the particle centers using Welford's method.
+							TVec3<T> CenterDelta = ElemBounds.Center() - WorkSnapshot.AverageCenter;
+							WorkSnapshot.AverageCenter += CenterDelta / (T)NumElems;
+							WorkSnapshot.ScaledCenterVariance += (ElemBounds.Center() - WorkSnapshot.AverageCenter) * CenterDelta;
+						}
+					}
+					else
+					{
+						ElemBounds = TAABB<T, 3>(TVec3<T>(TNumericLimits<T>::Lowest()), TVec3<T>(TNumericLimits<T>::Max()));
 					}
-				}
-				else
-				{
-					ElemBounds = TAABB<T, 3>(TVec3<T>(TNumericLimits<T>::Lowest()), TVec3<T>(TNumericLimits<T>::Max()));
-				}
 
-				if (!bHasBoundingBox)
-				{
-					if (bMutable)
+					if (!bHasBoundingBox)
 					{
-						PayloadToInfo.Add(Payload, FAABBTreePayloadInfo{ GlobalPayloads.Num(), INDEX_NONE, INDEX_NONE, INDEX_NONE });
+						if (bMutable)
+						{
+							PayloadToInfo.Add(Payload, FAABBTreePayloadInfo{ GlobalPayloads.Num(), INDEX_NONE, INDEX_NONE, INDEX_NONE });
+						}
+						GlobalPayloads.Add(FElement{ Payload, ElemBounds });
 					}
-					GlobalPayloads.Add(FElement{ Payload, ElemBounds });
-				}
 
-				++Idx;
-				//todo: payload info
+					++Idx;
+					//todo: payload info
+				}
 			}
 		}
 
-		NumProcessedThisSlice = Particles.Num();	//todo: give chance to time slice out of next phase
+		NumProcessedThisSlice = NumParticles;	//todo: give chance to time slice out of next phase
 
 		{
 			SCOPE_CYCLE_COUNTER(STAT_AABBTreeInitialTimeSlice);
diff --git a/Engine/Source/Runtime/Experimental/Chaos/Public/Chaos/ISpatialAcceleration.h b/Engine/Source/Runtime/Experimental/Chaos/Public/Chaos/ISpatialAcceleration.h
--- a/Engine/Source/Runtime/Experimental/Chaos/Public/Chaos/ISpatialAcceleration.h
+++ b/Engine/Source/Runtime/Experimental/Chaos/Public/Chaos/ISpatialAcceleration.h
@@ -94,7 +94,7 @@ public:
 	In production, this class will only contain a payload.
 */
 template <typename TPayloadType>
-struct CHAOS_API TSpatialVisitorData
+struct TSpatialVisitorData
 {
 	TPayloadType Payload;
 	TSpatialVisitorData(const TPayloadType& InPayload, const bool bInHasBounds = false, const FAABB3& InBounds = FAABB3::ZeroAABB())
@@ -115,9 +115,10 @@ struct CHAOS_API TSpatialVisitorData
 	This class determines whether the acceleration structure should continue to iterate through potential instances
 */
 template <typename TPayloadType, typename T = FReal>
-class CHAOS_API ISpatialVisitor
+class ISpatialVisitor
 {
 public:
+
 	virtual ~ISpatialVisitor() = default;
 
 	/** Called whenever an instance in the acceleration structure may overlap
@@ -251,12 +252,23 @@ FChaosArchive& operator<<(FChaosArchive& Ar, TPayloadBoundsElement<TPayloadType,
 	return Ar;
 }
 
+template<typename TPayloadType, typename T, int d>
+class ISpatialAcceleration;
+
+template<ESpatialAcceleration SpatialType, typename TPayloadType, typename T, int d>
+struct TSpatialAccelerationSerializationFactory
+{
+	static ISpatialAcceleration<TPayloadType, T, d>* Create();
+};
+
 template <typename TPayloadType, typename T, int d>
-class CHAOS_API ISpatialAcceleration
+class ISpatialAcceleration
 {
 public:
 	UE_NONCOPYABLE(ISpatialAcceleration)
 
+	using TPayload = TPayloadType;
+
 	ISpatialAcceleration(SpatialAccelerationType InType = static_cast<SpatialAccelerationType>(ESpatialAcceleration::Unknown))
 		: Type(InType), SyncTimestamp(0), AsyncTimeSlicingComplete(true)
 	{}
@@ -331,7 +343,26 @@ public:
 	virtual void DumpStats() const {}
 #endif
 
-	static ISpatialAcceleration<TPayloadType, T, d>* SerializationFactory(FChaosArchive& Ar, ISpatialAcceleration<TPayloadType, T, d>* Accel);
+	static ISpatialAcceleration<TPayloadType, T, d>* SerializationFactory(FChaosArchive& Ar, ISpatialAcceleration<TPayloadType, T, d>* Accel)
+	{
+		if (Ar.CustomVer(FExternalPhysicsCustomObjectVersion::GUID) < FExternalPhysicsCustomObjectVersion::SerializeEvolutionGenericAcceleration)
+		{
+			return TSpatialAccelerationSerializationFactory<ESpatialAcceleration::BoundingVolume, TPayloadType, T, d>::Create();
+		}
+
+		int8 AccelType = Ar.IsLoading() ? 0 : (int8)Accel->Type;
+		Ar << AccelType;
+		switch ((ESpatialAcceleration)AccelType)
+		{
+		case ESpatialAcceleration::BoundingVolume: return Ar.IsLoading() ? TSpatialAccelerationSerializationFactory<ESpatialAcceleration::BoundingVolume, TPayloadType, T, d>::Create() : nullptr;
+		case ESpatialAcceleration::AABBTree: return Ar.IsLoading() ? TSpatialAccelerationSerializationFactory<ESpatialAcceleration::AABBTree, TPayloadType, T, d>::Create() : nullptr;
+		case ESpatialAcceleration::AABBTreeBV: return Ar.IsLoading() ? TSpatialAccelerationSerializationFactory<ESpatialAcceleration::AABBTreeBV, TPayloadType, T, d>::Create() : nullptr;
+		case ESpatialAcceleration::Collection: check(false);	//Collections must be serialized directly since they are variadic
+		default: check(false); return nullptr;
+		}
+		return nullptr;
+	}
+
 	virtual void Serialize(FChaosArchive& Ar)
 	{
 		check(false);
@@ -685,5 +716,12 @@ FORCEINLINE bool PrePreFilterHelper(const int32 Payload, const TVisitor& Visitor
 	return false;
 }
 
+#if PLATFORM_MAC || PLATFORM_LINUX
+extern template class CHAOS_API ISpatialAcceleration<int32, FReal, 3>;
+extern template class CHAOS_API ISpatialVisitor<int32, FReal>;
+#else
+extern template class ISpatialAcceleration<int32, FReal, 3>;
+extern template class ISpatialVisitor<int32, FReal>;
+#endif
 
 }
diff --git a/Engine/Source/Runtime/Experimental/Chaos/Public/Chaos/ParticleHandle.h b/Engine/Source/Runtime/Experimental/Chaos/Public/Chaos/ParticleHandle.h
--- a/Engine/Source/Runtime/Experimental/Chaos/Public/Chaos/ParticleHandle.h
+++ b/Engine/Source/Runtime/Experimental/Chaos/Public/Chaos/ParticleHandle.h
@@ -3,6 +3,7 @@
 
 #include "Templates/ChooseClass.h"
 #include "Math/NumericLimits.h"
+#include "Chaos/ISpatialAcceleration.h"
 #include "Chaos/PBDRigidClusteredParticles.h"
 #include "Chaos/PBDGeometryCollectionParticles.h"
 #include "Chaos/ParticleHandleFwd.h"
@@ -192,6 +193,8 @@ FORCEINLINE_DEBUGGABLE bool PrePreSimFilterImp(const FCollisionFilterData& SimFi
 class FAccelerationStructureHandle
 {
 public:
+	static constexpr bool bHasPayloadOnInternalThread = true;
+
 	FAccelerationStructureHandle(FGeometryParticleHandle* InHandle);
 	FAccelerationStructureHandle(FGeometryParticle* InGeometryParticle = nullptr);
 
@@ -3476,5 +3479,14 @@ inline void SetObjectStateHelper(IPhysicsProxyBase& Proxy, FPBDRigidParticle& Ri
 
 CHAOS_API void SetObjectStateHelper(IPhysicsProxyBase& Proxy, FPBDRigidParticleHandle& Rigid, EObjectStateType InState, bool bAllowEvents = false, bool bInvalidate = true);
 
+#if PLATFORM_MAC || PLATFORM_LINUX
+extern template class CHAOS_API ISpatialAcceleration<FAccelerationStructureHandle, FReal, 3>;
+extern template class CHAOS_API ISpatialVisitor<FAccelerationStructureHandle, FReal>;
+#else
+extern template class ISpatialAcceleration<FAccelerationStructureHandle, FReal, 3>;
+extern template class ISpatialVisitor<FAccelerationStructureHandle, FReal>;
+#endif
+
+
 } // namespace Chaos
 
diff --git a/Engine/Source/Runtime/PhysicsCore/Private/SQAccelerator.cpp b/Engine/Source/Runtime/PhysicsCore/Private/SQAccelerator.cpp
--- a/Engine/Source/Runtime/PhysicsCore/Private/SQAccelerator.cpp
+++ b/Engine/Source/Runtime/PhysicsCore/Private/SQAccelerator.cpp
@@ -1,6 +1,7 @@
 // Copyright Epic Games, Inc. All Rights Reserved.
 
 #include "SQAccelerator.h"
+#include "SQVisitor.h"
 #include "CollisionQueryFilterCallbackCore.h"
 #include "PhysicsInterfaceUtilsCore.h"
 
@@ -14,11 +15,6 @@
 #include "Chaos/GeometryQueries.h"
 #include "Chaos/DebugDrawQueue.h"
 
-#if CHAOS_DEBUG_DRAW
-int32 ChaosSQDrawDebugVisitorQueries = 0;
-FAutoConsoleVariableRef CVarChaosSQDrawDebugQueries(TEXT("p.Chaos.SQ.DrawDebugVisitorQueries"), ChaosSQDrawDebugVisitorQueries, TEXT("Draw bounds of objects visited by visitors in scene queries."));
-#endif
-
 void FSQAcceleratorUnion::Raycast(const FVector& Start, const FVector& Dir, const float DeltaMagnitude, FPhysicsHitCallback<FHitRaycast>& HitBuffer, EHitFlags OutputFlags, const ChaosInterface::FQueryFilterData& QueryFilterData, ICollisionQueryFilterCallbackBase& QueryCallback) const
 {
 	for(const ISQAccelerator* Accelerator : Accelerators)
@@ -63,542 +59,6 @@ struct FPreFilterInfo
 	int32 ActorIdx;
 };
 
-template <typename TLocationHit>
-void FillHitHelperImp(TLocationHit& Hit, const Chaos::FReal Distance, const FVector& WorldPosition, const FVector& WorldNormal, int32 FaceIdx, bool bComputeMTD)
-{
-	Hit.Distance = Distance;
-	Hit.WorldPosition = WorldPosition;
-	Hit.WorldNormal = WorldNormal;
-	Hit.Flags = Distance > 0.f || bComputeMTD ? EHitFlags::Distance | EHitFlags::Normal | EHitFlags::Position : EHitFlags::Distance | EHitFlags::FaceIndex;
-	Hit.FaceIndex = FaceIdx;
-}
-
-void FillHitHelper(ChaosInterface::FLocationHit& Hit, const Chaos::FReal Distance, const FVector& WorldPosition, const FVector& WorldNormal, int32 FaceIdx, bool bComputeMTD)
-{
-	FillHitHelperImp(Hit, Distance, WorldPosition, WorldNormal, FaceIdx, bComputeMTD);
-}
-
-void FillHitHelper(ChaosInterface::FPTLocationHit& Hit, const Chaos::FReal Distance, const FVector& WorldPosition, const FVector& WorldNormal, int32 FaceIdx, bool bComputeMTD)
-{
-	FillHitHelperImp(Hit, Distance, WorldPosition, WorldNormal, FaceIdx, bComputeMTD);
-}
-
-void FillHitHelper(ChaosInterface::FOverlapHit& Hit, const Chaos::FReal Distance, const FVector& WorldPosition, const FVector& WorldNormal, int32 FaceIdx, bool bComputeMTD)
-{}
-
-void FillHitHelper(ChaosInterface::FPTOverlapHit& Hit, const Chaos::FReal Distance, const FVector& WorldPosition, const FVector& WorldNormal, int32 FaceIdx, bool bComputeMTD)
-{}
-
-template <typename QueryGeometryType, typename TPayload, typename THitType, bool bGTData = true>
-struct TSQVisitor : public Chaos::ISpatialVisitor<TPayload, Chaos::FReal>
-{
-	using TGeometryType = typename TChooseClass<bGTData, Chaos::FGeometryParticle, Chaos::FGeometryParticleHandle>::Result;
-	TSQVisitor(const FVector& InStartPoint, const FVector& InDir, ChaosInterface::FSQHitBuffer<THitType>& InHitBuffer, EHitFlags InOutputFlags,
-			   const ChaosInterface::FQueryFilterData& InQueryFilterData, ICollisionQueryFilterCallbackBase& InQueryCallback, const ChaosInterface::FQueryDebugParams& InDebugParams)
-		: StartPoint(InStartPoint)
-		, Dir(InDir)
-		, HalfExtents(0)
-		, OutputFlags(InOutputFlags)
-		, bAnyHit(false)
-		, DebugParams(InDebugParams)
-		, HitBuffer(InHitBuffer)
-		, QueryFilterData(InQueryFilterData)
-		, QueryFilterDataConcrete(C2UFilterData(QueryFilterData.data))
-		, QueryCallback(InQueryCallback)
-	{
-		bAnyHit = QueryFilterData.flags & FChaosQueryFlag::eANY_HIT;
-	}
-
-	TSQVisitor(const FTransform& InStartTM, const FVector& InDir, ChaosInterface::FSQHitBuffer<THitType>& InHitBuffer, EHitFlags InOutputFlags,
-			   const ChaosInterface::FQueryFilterData& InQueryFilterData, ICollisionQueryFilterCallbackBase& InQueryCallback, const QueryGeometryType& InQueryGeom, const ChaosInterface::FQueryDebugParams& InDebugParams)
-		: Dir(InDir)
-		, HalfExtents(InQueryGeom.BoundingBox().Extents() * 0.5)
-		, OutputFlags(InOutputFlags)
-		, bAnyHit(false)
-		, DebugParams(InDebugParams)
-		, HitFaceNormal(Chaos::FVec3::ZeroVector)
-		, HitBuffer(InHitBuffer)
-		, QueryFilterData(InQueryFilterData)
-		, QueryFilterDataConcrete(C2UFilterData(QueryFilterData.data))
-		, QueryGeom(&InQueryGeom)
-		, QueryCallback(InQueryCallback)
-		, StartTM(InStartTM)
-	{
-		bAnyHit = QueryFilterData.flags & FChaosQueryFlag::eANY_HIT;
-		//todo: check THitType is sweep
-	}
-
-	TSQVisitor(const FTransform& InWorldTM, ChaosInterface::FSQHitBuffer<THitType>& InHitBuffer,
-			   const ChaosInterface::FQueryFilterData& InQueryFilterData, ICollisionQueryFilterCallbackBase& InQueryCallback, const QueryGeometryType& InQueryGeom, const ChaosInterface::FQueryDebugParams& InDebugParams)
-		: HalfExtents(InQueryGeom.BoundingBox().Extents() * 0.5)
-		, bAnyHit(false)
-		, DebugParams(InDebugParams)
-		, HitFaceNormal(Chaos::FVec3::ZeroVector)
-		, HitBuffer(InHitBuffer)
-		, QueryFilterData(InQueryFilterData)
-		, QueryFilterDataConcrete(C2UFilterData(QueryFilterData.data))
-		, QueryGeom(&InQueryGeom)
-		, QueryCallback(InQueryCallback)
-		, StartTM(InWorldTM)
-	{
-		bAnyHit = QueryFilterData.flags & FChaosQueryFlag::eANY_HIT;
-		//todo: check THitType is overlap
-	}
-
-	virtual bool Raycast(const Chaos::TSpatialVisitorData<TPayload>& Instance, Chaos::FQueryFastData& CurData) override
-	{
-		return Visit<ESQType::Raycast>(Instance, &CurData);
-	}
-
-	virtual bool Sweep(const Chaos::TSpatialVisitorData<TPayload>& Instance, Chaos::FQueryFastData& CurData) override
-	{
-		return Visit<ESQType::Sweep>(Instance, &CurData);
-	}
-
-	virtual bool Overlap(const Chaos::TSpatialVisitorData<TPayload>& Instance) override
-	{
-		return Visit<ESQType::Overlap>(Instance, nullptr);
-	}
-
-	virtual const void* GetQueryData() const override
-	{
-		return &QueryFilterDataConcrete;
-	}
-
-	virtual const void* GetSimData() const override
-	{
-		return nullptr;
-	}
-
-	/** Return a pointer to the payload on which we are querying the acceleration structure */
-	virtual const void* GetQueryPayload() const override
-	{
-		return nullptr;
-	}
-
-	virtual bool HasBlockingHit() const override
-	{
-		return HitBuffer.HasBlockingHit();
-	}
-
-private:
-
-	enum class ESQType
-	{
-		Raycast,
-		Sweep,
-		Overlap
-	};
-
-	auto GetPayloadForThread(const TPayload& Payload)
-	{
-		if constexpr(bGTData)
-		{
-			return Payload.GetExternalGeometryParticle_ExternalThread();
-		}
-		else
-		{
-			return Payload.GetGeometryParticleHandle_PhysicsThread();
-		}
-	}
-
-	template <ESQType SQ>
-	bool Visit(const Chaos::TSpatialVisitorData<TPayload>& Instance, Chaos::FQueryFastData* CurData)
-	{
-		//QUICK_SCOPE_CYCLE_COUNTER(SQVisit);
-		TPayload Payload = Instance.Payload;
-
-		//todo: add a check to ensure hitbuffer matches SQ type
-		using namespace Chaos;
-		TGeometryType* GeometryParticle = GetPayloadForThread(Payload);
-
-		if(!GeometryParticle)
-		{
-			// This case handles particles created by the physics simulation without the main thread
-			// being made aware of their creation. We have a PT particle but no external particle
-			ensure(bGTData);
-
-			return true;
-		}
-
-		const FShapesArray& Shapes = GeometryParticle->ShapesArray();
-
-		const bool bTestShapeBounds = Shapes.Num() > 1;
-		bool bContinue = true;
-
-		const FRigidTransform3 ActorTM(GeometryParticle->X(), GeometryParticle->R());
-		const TAABB<FReal, 3> QueryGeomWorldBounds = QueryGeom ? QueryGeom->CalculateTransformedBounds(StartTM) : TAABB<FReal, 3>(-HalfExtents, HalfExtents);
-
-#if CHAOS_DEBUG_DRAW
-		bool bAllShapesIgnoredInPrefilter = true;
-		bool bHitBufferIncreased = false;
-#endif
-
-		for(const auto& Shape : Shapes)
-		{
-			const FImplicitObject* Geom = Shape->GetGeometry().Get();
-
-			if(bTestShapeBounds)
-			{
-				FAABB3 InflatedWorldBounds;
-				if(SQ == ESQType::Raycast)
-				{
-					InflatedWorldBounds = Shape->GetWorldSpaceInflatedShapeBounds();
-				}
-				else
-				{
-					// Transform to world bounds and get the proper half extent.
-					const FVec3 WorldHalfExtent = QueryGeom ? QueryGeomWorldBounds.Extents() * 0.5f : FVec3(HalfExtents);
-
-					InflatedWorldBounds = FAABB3(Shape->GetWorldSpaceInflatedShapeBounds().Min() - WorldHalfExtent, Shape->GetWorldSpaceInflatedShapeBounds().Max() + WorldHalfExtent);
-				}
-				if(SQ != ESQType::Overlap)
-				{
-					//todo: use fast raycast
-					Chaos::FReal TmpTime, TmpExitTime;
-					const FVec3 InflatedBoundsTraceStart = SQ == ESQType::Raycast ? StartPoint : QueryGeomWorldBounds.Center();
-					if(!InflatedWorldBounds.RaycastFast(InflatedBoundsTraceStart, CurData->Dir, CurData->InvDir, CurData->bParallel, CurData->CurrentLength, CurData->InvCurrentLength, TmpTime, TmpExitTime))
-					{
-						continue;
-					}
-				}
-				else
-				{
-					const FVec3 QueryCenter = QueryGeom ? QueryGeomWorldBounds.Center() : StartTM.GetLocation();
-					if(!InflatedWorldBounds.Contains(QueryCenter))
-					{
-						continue;
-					}
-				}
-			}
-
-			//TODO: use gt particles directly
-			ECollisionQueryHitType HitType = QueryFilterData.flags & FChaosQueryFlag::ePREFILTER ? QueryCallback.PreFilter(QueryFilterDataConcrete, *Shape, *GeometryParticle) : ECollisionQueryHitType::Block;
-
-			if(HitType != ECollisionQueryHitType::None)
-			{
-#if CHAOS_DEBUG_DRAW
-				bAllShapesIgnoredInPrefilter = false;
-#endif
-
-				//QUICK_SCOPE_CYCLE_COUNTER(SQNarrow);
-				THitType Hit;
-				Hit.Actor = GeometryParticle;
-				Hit.Shape = Shape.Get();
-
-				bool bHit = false;
-
-				FVec3 WorldPosition, WorldNormal;
-				Chaos::FReal Distance = 0;	//not needed but fixes compiler warning for overlap
-				int32 FaceIdx = INDEX_NONE;	//not needed but fixes compiler warning for overlap
-				FVec3 FaceNormal = FVec3::ZeroVector;
-				const bool bComputeMTD = !!((uint16)(OutputFlags.HitFlags & EHitFlags::MTD));
-
-				if(SQ == ESQType::Raycast)
-				{
-					FVec3 LocalNormal;
-					FVec3 LocalPosition;
-
-					const FVec3 DirLocal = ActorTM.InverseTransformVectorNoScale(Dir);
-					const FVec3 StartLocal = ActorTM.InverseTransformPositionNoScale(StartPoint);
-					bHit = Geom->Raycast(StartLocal, DirLocal, CurData->CurrentLength, /*Thickness=*/0.0, Distance, LocalPosition, LocalNormal, FaceIdx);
-					if(bHit)
-					{
-						WorldPosition = ActorTM.TransformPositionNoScale(LocalPosition);
-						WorldNormal = ActorTM.TransformVectorNoScale(LocalNormal);
-					}
-				}
-				else if(SQ == ESQType::Sweep && CurData->CurrentLength > 0 && ensure(QueryGeom))
-				{
-					bHit = SweepQuery(*Geom, ActorTM, *QueryGeom, StartTM, CurData->Dir, CurData->CurrentLength, Distance, WorldPosition, WorldNormal, FaceIdx, FaceNormal, 0.f, bComputeMTD);
-				}
-				else if((SQ == ESQType::Overlap || (SQ == ESQType::Sweep && CurData->CurrentLength == 0)) && ensure(QueryGeom))
-				{
-					if(bComputeMTD)
-					{
-						FMTDInfo MTDInfo;
-						bHit = OverlapQuery(*Geom, ActorTM, *QueryGeom, StartTM, /*Thickness=*/0, &MTDInfo);
-						if(bHit)
-						{
-							WorldNormal = MTDInfo.Normal * MTDInfo.Penetration;
-						}
-					}
-					else
-					{
-						bHit = OverlapQuery(*Geom, ActorTM, *QueryGeom, StartTM, /*Thickness=*/0);
-					}
-				}
-
-				if(bHit)
-				{
-					//QUICK_SCOPE_CYCLE_COUNTER(SQNarrowHit);
-
-					bool bAcceptHit = true;
-					if constexpr(SQ == ESQType::Sweep && std::is_same_v<THitType, ChaosInterface::FSweepHit>)
-					{
-						const THitType* CurrentHit = HitBuffer.GetCurrentHit();
-						if(FaceIdx != INDEX_NONE || !FaceNormal.IsNearlyZero())
-						{
-							if(CurrentHit)
-							{
-								constexpr static FReal CoLocationEpsilon = 1e-6;
-								const FReal DistDelta = FMath::Abs(Distance - CurrentHit->Distance);
-
-								if(DistDelta < CoLocationEpsilon && !HitFaceNormal.IsNearlyZero())
-								{
-									// We already have a face hit from another triangle mesh - see if this one is better (more opposing the sweep)
-									const FReal OldDot = FVec3::DotProduct(CurData->Dir, HitFaceNormal);
-									const FReal NewDot = FVec3::DotProduct(CurData->Dir, FaceNormal);
-
-									if(NewDot < OldDot)
-									{
-										// More opposing
-										HitFaceNormal = FaceNormal;
-									}
-									else
-									{
-										// This hit is co-located but has a worse normal
-										bAcceptHit = false;
-									}
-								}
-							}
-
-							if(bAcceptHit)
-							{
-								// Record the new face normal
-								HitFaceNormal = FaceNormal;
-							}
-						}
-					}
-
-					if(bAcceptHit)
-					{
-						FillHitHelper(Hit, Distance, WorldPosition, WorldNormal, FaceIdx, bComputeMTD);
-
-						if constexpr (std::is_base_of_v<ChaosInterface::FQueryHit, THitType> || std::is_base_of_v<ChaosInterface::FPTQueryHit, THitType>)
-						{
-							Hit.FaceNormal = HitFaceNormal;
-						}
-
-						HitType = QueryFilterData.flags & FChaosQueryFlag::ePOSTFILTER ? QueryCallback.PostFilter(QueryFilterDataConcrete, Hit) : HitType;
-
-						if(HitType != ECollisionQueryHitType::None)
-						{
-
-							//overlap never blocks
-							const bool bBlocker = (HitType == ECollisionQueryHitType::Block || bAnyHit || HitBuffer.WantsSingleResult());
-							HitBuffer.InsertHit(Hit, bBlocker);
-#if CHAOS_DEBUG_DRAW
-							bHitBufferIncreased = true;
-#endif
-
-							if(bBlocker && SQ != ESQType::Overlap)
-							{
-								CurData->SetLength(FMath::Max((FReal)0., Distance));	//Max is needed for MTD which returns negative distance
-								if(CurData->CurrentLength == 0 && (SQ == ESQType::Raycast || HitBuffer.WantsSingleResult()))	//raycasts always fail with distance 0, sweeps only matter if we want multi overlaps
-								{
-									bContinue = false; //initial overlap so nothing will be better than this
-									break;
-								}
-							}
-
-							if(bAnyHit)
-							{
-								bContinue = false;
-								break;
-							}
-						}
-					}
-				}
-			}
-		}
-
-#if CHAOS_DEBUG_DRAW
-		if(DebugParams.IsDebugQuery() && ChaosSQDrawDebugVisitorQueries)
-		{
-			DebugDraw<SQ>(Instance, CurData, bAllShapesIgnoredInPrefilter, bHitBufferIncreased);
-		}
-#endif
-
-		return bContinue;
-	}
-
-#if CHAOS_DEBUG_DRAW
-
-	void DebugDrawPayloadImpl(const TPayload& Payload, const bool bExternal, const bool bHit, decltype(&TPayload::DebugDraw))
-	{
-		Payload.DebugDraw(bExternal, bHit);
-	}
-	void DebugDrawPayloadImpl(const TPayload& Payload, const bool bExternal, const bool bHit, ...)
-	{}
-	void DebugDrawPayload(const TPayload& Payload, const bool bExternal, const bool bHit)
-	{
-		DebugDrawPayloadImpl(Payload, bExternal, bHit, 0);
-	}
-
-	template <ESQType SQ>
-	void DebugDraw(const Chaos::TSpatialVisitorData<TPayload>& Instance, const Chaos::FQueryFastData* CurData, const bool bPrefiltered, const bool bHit)
-	{
-		if(SQ == ESQType::Raycast)
-		{
-			const FVector EndPoint = StartPoint + (Dir * CurData->CurrentLength);
-			Chaos::FDebugDrawQueue::GetInstance().DrawDebugDirectionalArrow(StartPoint, EndPoint, 5.f, bHit ? FColor::Red : FColor::Green);
-		}
-		else if(SQ == ESQType::Overlap)
-		{
-			Chaos::DebugDraw::DrawShape(StartTM, QueryGeom, Chaos::FShapeOrShapesArray(), bHit ? FColor::Red : FColor::Green);
-		}
-
-		if(Instance.bHasBounds)
-		{
-			Chaos::FDebugDrawQueue::GetInstance().DrawDebugBox(Instance.Bounds.Center(), Instance.Bounds.Extents(), FQuat::Identity, bHit ? FColor(100, 50, 50) : FColor(50, 100, 50), false, -1.f, 0, 0.f);
-		}
-
-		if(!bPrefiltered)
-		{
-			DebugDrawPayload(Instance.Payload, DebugParams.bExternalQuery, bHit);
-		}
-	}
-#endif
-
-	const FVector StartPoint;
-	const FVector Dir;
-	const FVector HalfExtents;
-	FHitFlags OutputFlags;
-	bool bAnyHit;
-	const ChaosInterface::FQueryDebugParams DebugParams;
-	Chaos::FVec3 HitFaceNormal;
-	ChaosInterface::FSQHitBuffer<THitType>& HitBuffer;
-	const ChaosInterface::FQueryFilterData& QueryFilterData;
-	const FCollisionFilterData QueryFilterDataConcrete;
-	const QueryGeometryType* QueryGeom = nullptr;
-	ICollisionQueryFilterCallbackBase& QueryCallback;
-	const FTransform StartTM;
-};
-
-template <typename QueryGeometryType, typename TPayload, typename THitType, bool bGTData = true>
-struct TBPVisitor : public Chaos::ISpatialVisitor<TPayload, Chaos::FReal>
-{
-	using TGeometryType = typename TChooseClass<bGTData, Chaos::FGeometryParticle, Chaos::FGeometryParticleHandle>::Result;
-
-	TBPVisitor(const FTransform& InWorldTM, ChaosInterface::FSQHitBuffer<THitType>& InHitBuffer,
-			   const ChaosInterface::FQueryFilterData& InQueryFilterData, ICollisionQueryFilterCallbackBase& InQueryCallback, const QueryGeometryType& InQueryGeom, const ChaosInterface::FQueryDebugParams& InDebugParams)
-		: HalfExtents(InQueryGeom.BoundingBox().Extents() * 0.5)
-		, bAnyHit(false)
-		, DebugParams(InDebugParams)
-		, HitBuffer(InHitBuffer)
-		, QueryFilterData(InQueryFilterData)
-		, QueryFilterDataConcrete(ToUnrealFilterData(InQueryFilterData.data))
-		, QueryGeom(&InQueryGeom)
-		, QueryCallback(InQueryCallback)
-		, StartTM(InWorldTM)
-	{
-		bAnyHit = QueryFilterData.flags & FPhysicsQueryFlag::eANY_HIT;
-	}
-	virtual bool Overlap(const Chaos::TSpatialVisitorData<TPayload>& Instance) override
-	{
-		return Visit<ESQType::Overlap>(Instance, nullptr);
-	}
-
-	virtual bool Raycast(const Chaos::TSpatialVisitorData<TPayload>& Instance, Chaos::FQueryFastData& CurData) override
-	{
-		ensure(false);
-		return false;
-	}
-
-	virtual bool Sweep(const Chaos::TSpatialVisitorData<TPayload>& Instance, Chaos::FQueryFastData& CurData) override
-	{
-		ensure(false);
-		return false;
-	}
-
-	virtual const void* GetQueryData() const override
-	{
-		return &QueryFilterDataConcrete;
-	}
-
-	virtual const void* GetSimData() const override
-	{
-		return nullptr;
-	}
-
-	/** Return a pointer to the payload on which we are querying the acceleration structure */
-	virtual const void* GetQueryPayload() const override
-	{
-		return nullptr;
-	}
-
-private:
-
-	enum class ESQType
-	{
-		Raycast,
-		Sweep,
-		Overlap
-	};
-
-	auto GetPayloadForThread(const TPayload& Payload)
-	{
-		if constexpr(bGTData)
-		{
-			return Payload.GetExternalGeometryParticle_ExternalThread();
-		}
-		else
-		{
-			return Payload.GetGeometryParticleHandle_PhysicsThread();
-		}
-	}
-
-	template <ESQType SQ>
-	bool Visit(const Chaos::TSpatialVisitorData<TPayload>& Instance, Chaos::FQueryFastData* CurData)
-	{
-		bool bContinue = true;
-		TPayload Payload = Instance.Payload;
-		//todo: add a check to ensure hitbuffer matches SQ type
-		using namespace Chaos;
-		TGeometryType* GeometryParticle = GetPayloadForThread(Payload);
-		if(!GeometryParticle)
-		{
-			// This case handles particles created by the physics simulation without the main thread
-			// being made aware of their creation. We have a PT particle but no external particle
-			ensure(bGTData);
-			return true;
-		}
-		const FShapesArray& Shapes = GeometryParticle->ShapesArray();
-		THitType Hit;
-		Hit.Actor = GeometryParticle;
-		for(const auto& Shape : Shapes)
-		{
-			ECollisionQueryHitType HitType = QueryFilterData.flags & FPhysicsQueryFlag::ePREFILTER ? QueryCallback.PreFilter(QueryFilterDataConcrete, *Shape, *GeometryParticle) : ECollisionQueryHitType::Block;
-			if(HitType != ECollisionQueryHitType::None)
-			{
-				const bool bBlocker = (HitType == ECollisionQueryHitType::Block || bAnyHit || HitBuffer.WantsSingleResult());
-				Hit.Shape = Shape.Get();
-				HitBuffer.InsertHit(Hit, bBlocker);
-				if(bAnyHit)
-				{
-					bContinue = false;
-				}
-				break;
-			}
-		}
-		return bContinue;
-	}
-
-	const FVector StartPoint;
-	const FVector Dir;
-	const FVector HalfExtents;
-	FHitFlags OutputFlags;
-	bool bAnyHit;
-	const ChaosInterface::FQueryDebugParams DebugParams;
-	ChaosInterface::FSQHitBuffer<THitType>& HitBuffer;
-	const ChaosInterface::FQueryFilterData& QueryFilterData;
-	const FCollisionFilterData QueryFilterDataConcrete;
-	const QueryGeometryType* QueryGeom;
-	ICollisionQueryFilterCallbackBase& QueryCallback;
-	const FTransform StartTM;
-};
-
 template <typename TRaycastHit>
 void FChaosSQAccelerator::RaycastImp(const FVector& Start, const FVector& Dir, const float DeltaMagnitude, ChaosInterface::FSQHitBuffer<TRaycastHit>& HitBuffer, EHitFlags OutputFlags, const ChaosInterface::FQueryFilterData& QueryFilterData, ICollisionQueryFilterCallbackBase& QueryCallback, const ChaosInterface::FQueryDebugParams& DebugParams) const
 {
@@ -621,32 +81,6 @@ void FChaosSQAccelerator::Raycast(const FVector& Start, const FVector& Dir, cons
 	RaycastImp(Start, Dir, DeltaMagnitude, HitBuffer, OutputFlags, QueryFilterData, QueryCallback, DebugParams);
 }
 
-template <typename QueryGeomType, typename TSweepHit>
-void SweepHelper(const QueryGeomType& QueryGeom, const Chaos::ISpatialAcceleration<Chaos::FAccelerationStructureHandle, Chaos::FReal, 3>& SpatialAcceleration, const FTransform& StartTM, const FVector& Dir, const float DeltaMagnitude, ChaosInterface::FSQHitBuffer<TSweepHit>& HitBuffer, EHitFlags OutputFlags, const ChaosInterface::FQueryFilterData& QueryFilterData, ICollisionQueryFilterCallbackBase& QueryCallback, const ChaosInterface::FQueryDebugParams& DebugParams)
-{
-	using namespace Chaos;
-	using namespace ChaosInterface;
-
-	const FAABB3 Bounds = QueryGeom.BoundingBox().TransformedAABB(StartTM);
-	const bool bSweepAsOverlap = DeltaMagnitude == 0;	//question: do we care about tiny sweeps?
-	TSQVisitor<QueryGeomType, FAccelerationStructureHandle, TSweepHit, std::is_same<TSweepHit, FSweepHit>::value> SweepVisitor(StartTM, Dir, HitBuffer, OutputFlags, QueryFilterData, QueryCallback, QueryGeom, DebugParams);
-
-	HitBuffer.IncFlushCount();
-
-	if(bSweepAsOverlap)
-	{
-		//fallback to overlap
-		SpatialAcceleration.Overlap(Bounds, SweepVisitor);
-	}
-	else
-	{
-		const FVector HalfExtents = Bounds.Extents() * 0.5f;
-		SpatialAcceleration.Sweep(Bounds.GetCenter(), Dir, DeltaMagnitude, HalfExtents, SweepVisitor);
-	}
-
-	HitBuffer.DecFlushCount();
-}
-
 void FChaosSQAccelerator::Sweep(const Chaos::FImplicitObject& QueryGeom, const FTransform& StartTM, const FVector& Dir, const float DeltaMagnitude, ChaosInterface::FSQHitBuffer<ChaosInterface::FSweepHit>& HitBuffer, EHitFlags OutputFlags, const ChaosInterface::FQueryFilterData& QueryFilterData, ICollisionQueryFilterCallbackBase& QueryCallback, const ChaosInterface::FQueryDebugParams& DebugParams) const
 {
 	return Chaos::Utilities::CastHelper(QueryGeom, StartTM, [&](const auto& Downcast, const FTransform& StartFullTM)
@@ -663,32 +97,6 @@ void FChaosSQAccelerator::Sweep(const Chaos::FImplicitObject& QueryGeom, const F
 	});
 }
 
-template <typename QueryGeomType, typename TOverlapHit>
-void OverlapHelper(const QueryGeomType& QueryGeom, const Chaos::ISpatialAcceleration<Chaos::FAccelerationStructureHandle, Chaos::FReal, 3>& SpatialAcceleration, const FTransform& GeomPose, ChaosInterface::FSQHitBuffer<TOverlapHit>& HitBuffer, const ChaosInterface::FQueryFilterData& QueryFilterData, ICollisionQueryFilterCallbackBase& QueryCallback, const ChaosInterface::FQueryDebugParams& DebugParams)
-{
-	using namespace Chaos;
-	using namespace ChaosInterface;
-
-	const FAABB3 Bounds = QueryGeom.CalculateTransformedBounds(GeomPose);
-
-	HitBuffer.IncFlushCount();
-
-	bool bSkipNarrowPhase = QueryFilterData.flags & FPhysicsQueryFlag::eSKIPNARROWPHASE;
-
-	constexpr bool bGTData = std::is_same<TOverlapHit, FOverlapHit>::value;
-	if(bSkipNarrowPhase)
-	{
-		TBPVisitor<QueryGeomType, FAccelerationStructureHandle, TOverlapHit, bGTData> OverlapVisitor(GeomPose, HitBuffer, QueryFilterData, QueryCallback, QueryGeom, DebugParams);
-		SpatialAcceleration.Overlap(Bounds, OverlapVisitor);
-	}
-	else
-	{
-		TSQVisitor<QueryGeomType, FAccelerationStructureHandle, TOverlapHit, bGTData> OverlapVisitor(GeomPose, HitBuffer, QueryFilterData, QueryCallback, QueryGeom, DebugParams);
-		SpatialAcceleration.Overlap(Bounds, OverlapVisitor);
-	}
-	HitBuffer.DecFlushCount();
-}
-
 void FChaosSQAccelerator::Overlap(const Chaos::FImplicitObject& QueryGeom, const FTransform& GeomPose, ChaosInterface::FSQHitBuffer<ChaosInterface::FOverlapHit>& HitBuffer, const ChaosInterface::FQueryFilterData& QueryFilterData, ICollisionQueryFilterCallbackBase& QueryCallback, const ChaosInterface::FQueryDebugParams& DebugParams) const
 {
 	return Chaos::Utilities::CastHelper(QueryGeom, GeomPose, [&](const auto& Downcast, const FTransform& GeomFullPose)
diff --git a/Engine/Source/Runtime/PhysicsCore/Private/SQVisitor.cpp b/Engine/Source/Runtime/PhysicsCore/Private/SQVisitor.cpp
new file mode 100644
--- /dev/null
+++ b/Engine/Source/Runtime/PhysicsCore/Private/SQVisitor.cpp
@@ -0,0 +1,8 @@
+// Copyright Epic Games, Inc. All Rights Reserved.
+
+#include "SQVisitor.h"
+
+#if CHAOS_DEBUG_DRAW
+int32 ChaosSQDrawDebugVisitorQueries = 0;
+FAutoConsoleVariableRef CVarChaosSQDrawDebugQueries(TEXT("p.Chaos.SQ.DrawDebugVisitorQueries"), ChaosSQDrawDebugVisitorQueries, TEXT("Draw bounds of objects visited by visitors in scene queries."));
+#endif
\ No newline at end of file
diff --git a/Engine/Source/Runtime/PhysicsCore/Public/SQVisitor.h b/Engine/Source/Runtime/PhysicsCore/Public/SQVisitor.h
new file mode 100644
--- /dev/null
+++ b/Engine/Source/Runtime/PhysicsCore/Public/SQVisitor.h
@@ -0,0 +1,611 @@
+// Copyright Epic Games, Inc. All Rights Reserved.
+
+#pragma once
+
+#include "Chaos/DebugDrawQueue.h"
+#include "Chaos/GeometryQueries.h"
+#include "Chaos/Interface/SQTypes.h"
+#include "Chaos/ISpatialAcceleration.h"
+#include "Chaos/ParticleHandle.h"
+#include "ChaosInterfaceWrapperCore.h"
+#include "CollisionQueryFilterCallbackCore.h"
+#include "PhysicsInterfaceTypesCore.h"
+#include "PhysicsInterfaceUtilsCore.h"
+
+#include <type_traits>
+
+#if CHAOS_DEBUG_DRAW
+extern PHYSICSCORE_API int32 ChaosSQDrawDebugVisitorQueries;
+extern PHYSICSCORE_API FAutoConsoleVariableRef CVarChaosSQDrawDebugQueries;
+#endif
+
+template <typename TLocationHit>
+void FillHitHelper(TLocationHit& Hit, const Chaos::FReal Distance, const FVector& WorldPosition, const FVector& WorldNormal, int32 FaceIdx, bool bComputeMTD)
+{
+	if constexpr (std::is_base_of_v<ChaosInterface::FLocationHit, TLocationHit> || std::is_base_of_v<ChaosInterface::FPTLocationHit, TLocationHit>)
+	{
+		Hit.Distance = Distance;
+		Hit.WorldPosition = WorldPosition;
+		Hit.WorldNormal = WorldNormal;
+		Hit.Flags = Distance > 0.f || bComputeMTD ? EHitFlags::Distance | EHitFlags::Normal | EHitFlags::Position : EHitFlags::Distance | EHitFlags::FaceIndex;
+		Hit.FaceIndex = FaceIdx;
+	}
+}
+
+template <typename QueryGeometryType, typename TPayload, typename THitType, bool bGTData = true>
+struct TSQVisitor : public Chaos::ISpatialVisitor<TPayload, Chaos::FReal>
+{
+	using TGeometryType = std::conditional_t<bGTData, Chaos::FGeometryParticle, Chaos::FGeometryParticleHandle>;
+	TSQVisitor(const FVector& InStartPoint, const FVector& InDir, ChaosInterface::FSQHitBuffer<THitType>& InHitBuffer, EHitFlags InOutputFlags,
+		const ChaosInterface::FQueryFilterData& InQueryFilterData, ICollisionQueryFilterCallbackBase& InQueryCallback, const ChaosInterface::FQueryDebugParams& InDebugParams)
+		: StartPoint(InStartPoint)
+		, Dir(InDir)
+		, HalfExtents(0)
+		, OutputFlags(InOutputFlags)
+		, bAnyHit(false)
+		, DebugParams(InDebugParams)
+		, HitBuffer(InHitBuffer)
+		, QueryFilterData(InQueryFilterData)
+		, QueryFilterDataConcrete(C2UFilterData(QueryFilterData.data))
+		, QueryCallback(InQueryCallback)
+	{
+		bAnyHit = QueryFilterData.flags & FChaosQueryFlag::eANY_HIT;
+	}
+
+	TSQVisitor(const FTransform& InStartTM, const FVector& InDir, ChaosInterface::FSQHitBuffer<THitType>& InHitBuffer, EHitFlags InOutputFlags,
+		const ChaosInterface::FQueryFilterData& InQueryFilterData, ICollisionQueryFilterCallbackBase& InQueryCallback, const QueryGeometryType& InQueryGeom, const ChaosInterface::FQueryDebugParams& InDebugParams)
+		: Dir(InDir)
+		, HalfExtents(InQueryGeom.BoundingBox().Extents() * 0.5)
+		, OutputFlags(InOutputFlags)
+		, bAnyHit(false)
+		, DebugParams(InDebugParams)
+		, HitFaceNormal(Chaos::FVec3::ZeroVector)
+		, HitBuffer(InHitBuffer)
+		, QueryFilterData(InQueryFilterData)
+		, QueryFilterDataConcrete(C2UFilterData(QueryFilterData.data))
+		, QueryGeom(&InQueryGeom)
+		, QueryCallback(InQueryCallback)
+		, StartTM(InStartTM)
+	{
+		bAnyHit = QueryFilterData.flags & FChaosQueryFlag::eANY_HIT;
+		//todo: check THitType is sweep
+	}
+
+	TSQVisitor(const FTransform& InWorldTM, ChaosInterface::FSQHitBuffer<THitType>& InHitBuffer,
+		const ChaosInterface::FQueryFilterData& InQueryFilterData, ICollisionQueryFilterCallbackBase& InQueryCallback, const QueryGeometryType& InQueryGeom, const ChaosInterface::FQueryDebugParams& InDebugParams)
+		: HalfExtents(InQueryGeom.BoundingBox().Extents() * 0.5)
+		, bAnyHit(false)
+		, DebugParams(InDebugParams)
+		, HitFaceNormal(Chaos::FVec3::ZeroVector)
+		, HitBuffer(InHitBuffer)
+		, QueryFilterData(InQueryFilterData)
+		, QueryFilterDataConcrete(C2UFilterData(QueryFilterData.data))
+		, QueryGeom(&InQueryGeom)
+		, QueryCallback(InQueryCallback)
+		, StartTM(InWorldTM)
+	{
+		bAnyHit = QueryFilterData.flags & FChaosQueryFlag::eANY_HIT;
+		//todo: check THitType is overlap
+	}
+
+	virtual ~TSQVisitor() {}
+
+	virtual bool Raycast(const Chaos::TSpatialVisitorData<TPayload>& Instance, Chaos::FQueryFastData& CurData) override
+	{
+		return Visit<ESQType::Raycast>(Instance, &CurData);
+	}
+
+	virtual bool Sweep(const Chaos::TSpatialVisitorData<TPayload>& Instance, Chaos::FQueryFastData& CurData) override
+	{
+		return Visit<ESQType::Sweep>(Instance, &CurData);
+	}
+
+	virtual bool Overlap(const Chaos::TSpatialVisitorData<TPayload>& Instance) override
+	{
+		return Visit<ESQType::Overlap>(Instance, nullptr);
+	}
+
+	virtual const void* GetQueryData() const override
+	{
+		return &QueryFilterDataConcrete;
+	}
+
+	virtual const void* GetSimData() const override
+	{
+		return nullptr;
+	}
+
+	/** Return a pointer to the payload on which we are querying the acceleration structure */
+	virtual const void* GetQueryPayload() const override
+	{
+		return nullptr;
+	}
+
+	virtual bool HasBlockingHit() const override
+	{
+		return HitBuffer.HasBlockingHit();
+	}
+
+private:
+
+	enum class ESQType
+	{
+		Raycast,
+		Sweep,
+		Overlap
+	};
+
+	TGeometryType* GetPayloadForThread(const TPayload& Payload)
+	{
+		if constexpr (bGTData)
+		{
+			return Payload.GetExternalGeometryParticle_ExternalThread();
+		}
+		else if constexpr (TPayload::bHasPayloadOnInternalThread)
+		{
+			return Payload.GetGeometryParticleHandle_PhysicsThread();
+		}
+		else
+		{
+			return nullptr;
+		}
+	}
+
+	template <ESQType SQ>
+	bool Visit(const Chaos::TSpatialVisitorData<TPayload>& Instance, Chaos::FQueryFastData* CurData)
+	{
+		//QUICK_SCOPE_CYCLE_COUNTER(SQVisit);
+		TPayload Payload = Instance.Payload;
+
+		//todo: add a check to ensure hitbuffer matches SQ type
+		using namespace Chaos;
+		TGeometryType* GeometryParticle = GetPayloadForThread(Payload);
+
+		if (!GeometryParticle)
+		{
+			// This case handles particles created by the physics simulation without the main thread
+			// being made aware of their creation. We have a PT particle but no external particle
+			ensure(bGTData);
+
+			return true;
+		}
+
+		const FShapesArray& Shapes = GeometryParticle->ShapesArray();
+
+		const bool bTestShapeBounds = Shapes.Num() > 1;
+		bool bContinue = true;
+
+		const FRigidTransform3 ActorTM(GeometryParticle->X(), GeometryParticle->R());
+		const TAABB<FReal, 3> QueryGeomWorldBounds = QueryGeom ? QueryGeom->CalculateTransformedBounds(StartTM) : TAABB<FReal, 3>(-HalfExtents, HalfExtents);
+
+#if CHAOS_DEBUG_DRAW
+		bool bAllShapesIgnoredInPrefilter = true;
+		bool bHitBufferIncreased = false;
+#endif
+
+		for (const auto& Shape : Shapes)
+		{
+			const FImplicitObject* Geom = Shape->GetGeometry().Get();
+
+			if (bTestShapeBounds)
+			{
+				FAABB3 InflatedWorldBounds;
+				if (SQ == ESQType::Raycast)
+				{
+					InflatedWorldBounds = Shape->GetWorldSpaceInflatedShapeBounds();
+				}
+				else
+				{
+					// Transform to world bounds and get the proper half extent.
+					const FVec3 WorldHalfExtent = QueryGeom ? QueryGeomWorldBounds.Extents() * 0.5f : FVec3(HalfExtents);
+
+					InflatedWorldBounds = FAABB3(Shape->GetWorldSpaceInflatedShapeBounds().Min() - WorldHalfExtent, Shape->GetWorldSpaceInflatedShapeBounds().Max() + WorldHalfExtent);
+				}
+				if (SQ != ESQType::Overlap)
+				{
+					//todo: use fast raycast
+					Chaos::FReal TmpTime, TmpExitTime;
+					const FVec3 InflatedBoundsTraceStart = SQ == ESQType::Raycast ? StartPoint : QueryGeomWorldBounds.Center();
+					if (!InflatedWorldBounds.RaycastFast(InflatedBoundsTraceStart, CurData->Dir, CurData->InvDir, CurData->bParallel, CurData->CurrentLength, CurData->InvCurrentLength, TmpTime, TmpExitTime))
+					{
+						continue;
+					}
+				}
+				else
+				{
+					const FVec3 QueryCenter = QueryGeom ? QueryGeomWorldBounds.Center() : StartTM.GetLocation();
+					if (!InflatedWorldBounds.Contains(QueryCenter))
+					{
+						continue;
+					}
+				}
+			}
+
+			//TODO: use gt particles directly
+			ECollisionQueryHitType HitType = QueryFilterData.flags & FChaosQueryFlag::ePREFILTER ? QueryCallback.PreFilter(QueryFilterDataConcrete, *Shape, *GeometryParticle) : ECollisionQueryHitType::Block;
+
+			if (HitType != ECollisionQueryHitType::None)
+			{
+#if CHAOS_DEBUG_DRAW
+				bAllShapesIgnoredInPrefilter = false;
+#endif
+
+				//QUICK_SCOPE_CYCLE_COUNTER(SQNarrow);
+				THitType Hit;
+				Hit.Actor = GeometryParticle;
+				Hit.Shape = Shape.Get();
+
+				bool bHit = false;
+
+				FVec3 WorldPosition, WorldNormal;
+				Chaos::FReal Distance = 0;	//not needed but fixes compiler warning for overlap
+				int32 FaceIdx = INDEX_NONE;	//not needed but fixes compiler warning for overlap
+				FVec3 FaceNormal = FVec3::ZeroVector;
+				const bool bComputeMTD = !!((uint16)(OutputFlags.HitFlags & EHitFlags::MTD));
+
+				if (SQ == ESQType::Raycast)
+				{
+					FVec3 LocalNormal;
+					FVec3 LocalPosition;
+
+					const FVec3 DirLocal = ActorTM.InverseTransformVectorNoScale(Dir);
+					const FVec3 StartLocal = ActorTM.InverseTransformPositionNoScale(StartPoint);
+					bHit = Geom->Raycast(StartLocal, DirLocal, CurData->CurrentLength, /*Thickness=*/0.0, Distance, LocalPosition, LocalNormal, FaceIdx);
+					if (bHit)
+					{
+						WorldPosition = ActorTM.TransformPositionNoScale(LocalPosition);
+						WorldNormal = ActorTM.TransformVectorNoScale(LocalNormal);
+					}
+				}
+				else if (SQ == ESQType::Sweep && CurData->CurrentLength > 0 && ensure(QueryGeom))
+				{
+					bHit = SweepQuery(*Geom, ActorTM, *QueryGeom, StartTM, CurData->Dir, CurData->CurrentLength, Distance, WorldPosition, WorldNormal, FaceIdx, FaceNormal, 0.f, bComputeMTD);
+				}
+				else if ((SQ == ESQType::Overlap || (SQ == ESQType::Sweep && CurData->CurrentLength == 0)) && ensure(QueryGeom))
+				{
+					if (bComputeMTD)
+					{
+						FMTDInfo MTDInfo;
+						bHit = OverlapQuery(*Geom, ActorTM, *QueryGeom, StartTM, /*Thickness=*/0, &MTDInfo);
+						if (bHit)
+						{
+							WorldNormal = MTDInfo.Normal * MTDInfo.Penetration;
+						}
+					}
+					else
+					{
+						bHit = OverlapQuery(*Geom, ActorTM, *QueryGeom, StartTM, /*Thickness=*/0);
+					}
+				}
+
+				if (bHit)
+				{
+					//QUICK_SCOPE_CYCLE_COUNTER(SQNarrowHit);
+
+					bool bAcceptHit = true;
+					if constexpr (SQ == ESQType::Sweep && std::is_same_v<THitType, ChaosInterface::FSweepHit>)
+					{
+						const THitType* CurrentHit = HitBuffer.GetCurrentHit();
+						if (FaceIdx != INDEX_NONE || !FaceNormal.IsNearlyZero())
+						{
+							if (CurrentHit)
+							{
+								constexpr static FReal CoLocationEpsilon = 1e-6;
+								const FReal DistDelta = FMath::Abs(Distance - CurrentHit->Distance);
+
+								if (DistDelta < CoLocationEpsilon && !HitFaceNormal.IsNearlyZero())
+								{
+									// We already have a face hit from another triangle mesh - see if this one is better (more opposing the sweep)
+									const FReal OldDot = FVec3::DotProduct(CurData->Dir, HitFaceNormal);
+									const FReal NewDot = FVec3::DotProduct(CurData->Dir, FaceNormal);
+
+									if (NewDot < OldDot)
+									{
+										// More opposing
+										HitFaceNormal = FaceNormal;
+									}
+									else
+									{
+										// This hit is co-located but has a worse normal
+										bAcceptHit = false;
+									}
+								}
+							}
+
+							if (bAcceptHit)
+							{
+								// Record the new face normal
+								HitFaceNormal = FaceNormal;
+							}
+						}
+					}
+
+					if (bAcceptHit)
+					{
+						FillHitHelper(Hit, Distance, WorldPosition, WorldNormal, FaceIdx, bComputeMTD);
+
+						if constexpr (std::is_base_of_v<ChaosInterface::FQueryHit, THitType> || std::is_base_of_v<ChaosInterface::FPTQueryHit, THitType>)
+						{
+							Hit.FaceNormal = HitFaceNormal;
+						}
+
+						HitType = QueryFilterData.flags & FChaosQueryFlag::ePOSTFILTER ? QueryCallback.PostFilter(QueryFilterDataConcrete, Hit) : HitType;
+
+						if (HitType != ECollisionQueryHitType::None)
+						{
+
+							//overlap never blocks
+							const bool bBlocker = (HitType == ECollisionQueryHitType::Block || bAnyHit || HitBuffer.WantsSingleResult());
+							HitBuffer.InsertHit(Hit, bBlocker);
+#if CHAOS_DEBUG_DRAW
+							bHitBufferIncreased = true;
+#endif
+
+							if (bBlocker && SQ != ESQType::Overlap)
+							{
+								CurData->SetLength(FMath::Max((FReal)0., Distance));	//Max is needed for MTD which returns negative distance
+								if (CurData->CurrentLength == 0 && (SQ == ESQType::Raycast || HitBuffer.WantsSingleResult()))	//raycasts always fail with distance 0, sweeps only matter if we want multi overlaps
+								{
+									bContinue = false; //initial overlap so nothing will be better than this
+									break;
+								}
+							}
+
+							if (bAnyHit)
+							{
+								bContinue = false;
+								break;
+							}
+						}
+					}
+				}
+			}
+		}
+
+#if CHAOS_DEBUG_DRAW
+		if (DebugParams.IsDebugQuery() && ChaosSQDrawDebugVisitorQueries)
+		{
+			DebugDraw<SQ>(Instance, CurData, bAllShapesIgnoredInPrefilter, bHitBufferIncreased);
+		}
+#endif
+
+		return bContinue;
+	}
+
+#if CHAOS_DEBUG_DRAW
+
+	void DebugDrawPayloadImpl(const TPayload& Payload, const bool bExternal, const bool bHit, decltype(&TPayload::DebugDraw))
+	{
+		Payload.DebugDraw(bExternal, bHit);
+	}
+	void DebugDrawPayloadImpl(const TPayload& Payload, const bool bExternal, const bool bHit, ...)
+	{}
+	void DebugDrawPayload(const TPayload& Payload, const bool bExternal, const bool bHit)
+	{
+		DebugDrawPayloadImpl(Payload, bExternal, bHit, 0);
+	}
+
+	template <ESQType SQ>
+	void DebugDraw(const Chaos::TSpatialVisitorData<TPayload>& Instance, const Chaos::FQueryFastData* CurData, const bool bPrefiltered, const bool bHit)
+	{
+		if (SQ == ESQType::Raycast)
+		{
+			const FVector EndPoint = StartPoint + (Dir * CurData->CurrentLength);
+			Chaos::FDebugDrawQueue::GetInstance().DrawDebugDirectionalArrow(StartPoint, EndPoint, 5.f, bHit ? FColor::Red : FColor::Green);
+		}
+		else if (SQ == ESQType::Overlap)
+		{
+			Chaos::DebugDraw::DrawShape(StartTM, QueryGeom, Chaos::FShapeOrShapesArray(), bHit ? FColor::Red : FColor::Green);
+		}
+
+		if (Instance.bHasBounds)
+		{
+			Chaos::FDebugDrawQueue::GetInstance().DrawDebugBox(Instance.Bounds.Center(), Instance.Bounds.Extents(), FQuat::Identity, bHit ? FColor(100, 50, 50) : FColor(50, 100, 50), false, -1.f, 0, 0.f);
+		}
+
+		if (!bPrefiltered)
+		{
+			DebugDrawPayload(Instance.Payload, DebugParams.bExternalQuery, bHit);
+		}
+	}
+#endif
+
+	const FVector StartPoint;
+	const FVector Dir;
+	const FVector HalfExtents;
+	FHitFlags OutputFlags;
+	bool bAnyHit;
+	const ChaosInterface::FQueryDebugParams DebugParams;
+	Chaos::FVec3 HitFaceNormal;
+	ChaosInterface::FSQHitBuffer<THitType>& HitBuffer;
+	const ChaosInterface::FQueryFilterData& QueryFilterData;
+	const FCollisionFilterData QueryFilterDataConcrete;
+	const QueryGeometryType* QueryGeom = nullptr;
+	ICollisionQueryFilterCallbackBase& QueryCallback;
+	const FTransform StartTM;
+};
+
+
+template <typename QueryGeometryType, typename TPayload, typename THitType, bool bGTData = true>
+struct TBPVisitor : public Chaos::ISpatialVisitor<TPayload, Chaos::FReal>
+{
+	using TGeometryType = typename TChooseClass<bGTData, Chaos::FGeometryParticle, Chaos::FGeometryParticleHandle>::Result;
+
+	TBPVisitor(const FTransform& InWorldTM, ChaosInterface::FSQHitBuffer<THitType>& InHitBuffer,
+		const ChaosInterface::FQueryFilterData& InQueryFilterData, ICollisionQueryFilterCallbackBase& InQueryCallback, const QueryGeometryType& InQueryGeom, const ChaosInterface::FQueryDebugParams& InDebugParams)
+		: HalfExtents(InQueryGeom.BoundingBox().Extents() * 0.5)
+		, bAnyHit(false)
+		, DebugParams(InDebugParams)
+		, HitBuffer(InHitBuffer)
+		, QueryFilterData(InQueryFilterData)
+		, QueryFilterDataConcrete(ToUnrealFilterData(InQueryFilterData.data))
+		, QueryGeom(&InQueryGeom)
+		, QueryCallback(InQueryCallback)
+		, StartTM(InWorldTM)
+	{
+		bAnyHit = QueryFilterData.flags & FPhysicsQueryFlag::eANY_HIT;
+	}
+
+	virtual ~TBPVisitor() {}
+
+	virtual bool Overlap(const Chaos::TSpatialVisitorData<TPayload>& Instance) override
+	{
+		return Visit<ESQType::Overlap>(Instance, nullptr);
+	}
+
+	virtual bool Raycast(const Chaos::TSpatialVisitorData<TPayload>& Instance, Chaos::FQueryFastData& CurData) override
+	{
+		ensure(false);
+		return false;
+	}
+
+	virtual bool Sweep(const Chaos::TSpatialVisitorData<TPayload>& Instance, Chaos::FQueryFastData& CurData) override
+	{
+		ensure(false);
+		return false;
+	}
+
+	virtual const void* GetQueryData() const override
+	{
+		return &QueryFilterDataConcrete;
+	}
+
+	virtual const void* GetSimData() const override
+	{
+		return nullptr;
+	}
+
+	/** Return a pointer to the payload on which we are querying the acceleration structure */
+	virtual const void* GetQueryPayload() const override
+	{
+		return nullptr;
+	}
+
+private:
+
+	enum class ESQType
+	{
+		Raycast,
+		Sweep,
+		Overlap
+	};
+
+	TGeometryType* GetPayloadForThread(const TPayload& Payload)
+	{
+		if constexpr (bGTData)
+		{
+			return Payload.GetExternalGeometryParticle_ExternalThread();
+		}
+		else if constexpr (TPayload::bHasPayloadOnInternalThread)
+		{
+			return Payload.GetGeometryParticleHandle_PhysicsThread();
+		}
+		else
+		{
+			return nullptr;
+		}
+	}
+
+	template <ESQType SQ>
+	bool Visit(const Chaos::TSpatialVisitorData<TPayload>& Instance, Chaos::FQueryFastData* CurData)
+	{
+		bool bContinue = true;
+		TPayload Payload = Instance.Payload;
+		//todo: add a check to ensure hitbuffer matches SQ type
+		using namespace Chaos;
+		TGeometryType* GeometryParticle = GetPayloadForThread(Payload);
+		if (!GeometryParticle)
+		{
+			// This case handles particles created by the physics simulation without the main thread
+			// being made aware of their creation. We have a PT particle but no external particle
+			ensure(bGTData);
+			return true;
+		}
+		const FShapesArray& Shapes = GeometryParticle->ShapesArray();
+		THitType Hit;
+		Hit.Actor = GeometryParticle;
+		for (const auto& Shape : Shapes)
+		{
+			ECollisionQueryHitType HitType = QueryFilterData.flags & FPhysicsQueryFlag::ePREFILTER ? QueryCallback.PreFilter(QueryFilterDataConcrete, *Shape, *GeometryParticle) : ECollisionQueryHitType::Block;
+			if (HitType != ECollisionQueryHitType::None)
+			{
+				const bool bBlocker = (HitType == ECollisionQueryHitType::Block || bAnyHit || HitBuffer.WantsSingleResult());
+				Hit.Shape = Shape.Get();
+				HitBuffer.InsertHit(Hit, bBlocker);
+				if (bAnyHit)
+				{
+					bContinue = false;
+				}
+				break;
+			}
+		}
+		return bContinue;
+	}
+
+	const FVector StartPoint;
+	const FVector Dir;
+	const FVector HalfExtents;
+	FHitFlags OutputFlags;
+	bool bAnyHit;
+	const ChaosInterface::FQueryDebugParams DebugParams;
+	ChaosInterface::FSQHitBuffer<THitType>& HitBuffer;
+	const ChaosInterface::FQueryFilterData& QueryFilterData;
+	const FCollisionFilterData QueryFilterDataConcrete;
+	const QueryGeometryType* QueryGeom;
+	ICollisionQueryFilterCallbackBase& QueryCallback;
+	const FTransform StartTM;
+};
+
+
+template <typename QueryGeomType, typename TSweepHit, typename TPayload>
+void SweepHelper(const QueryGeomType& QueryGeom, const Chaos::ISpatialAcceleration<TPayload, Chaos::FReal, 3>& SpatialAcceleration, const FTransform& StartTM, const FVector& Dir, const float DeltaMagnitude, ChaosInterface::FSQHitBuffer<TSweepHit>& HitBuffer, EHitFlags OutputFlags, const ChaosInterface::FQueryFilterData& QueryFilterData, ICollisionQueryFilterCallbackBase& QueryCallback, const ChaosInterface::FQueryDebugParams& DebugParams)
+{
+	using namespace Chaos;
+	using namespace ChaosInterface;
+
+	const FAABB3 Bounds = QueryGeom.BoundingBox().TransformedAABB(StartTM);
+	const bool bSweepAsOverlap = DeltaMagnitude == 0;	//question: do we care about tiny sweeps?
+	TSQVisitor<QueryGeomType, TPayload, TSweepHit, std::is_same<TSweepHit, FSweepHit>::value> SweepVisitor(StartTM, Dir, HitBuffer, OutputFlags, QueryFilterData, QueryCallback, QueryGeom, DebugParams);
+
+	HitBuffer.IncFlushCount();
+
+	if (bSweepAsOverlap)
+	{
+		//fallback to overlap
+		SpatialAcceleration.Overlap(Bounds, SweepVisitor);
+	}
+	else
+	{
+		const FVector HalfExtents = Bounds.Extents() * 0.5f;
+		SpatialAcceleration.Sweep(Bounds.GetCenter(), Dir, DeltaMagnitude, HalfExtents, SweepVisitor);
+	}
+
+	HitBuffer.DecFlushCount();
+}
+
+
+template <typename QueryGeomType, typename TOverlapHit, typename TPayload>
+void OverlapHelper(const QueryGeomType& QueryGeom, const Chaos::ISpatialAcceleration<TPayload, Chaos::FReal, 3>& SpatialAcceleration, const FTransform& GeomPose, ChaosInterface::FSQHitBuffer<TOverlapHit>& HitBuffer, const ChaosInterface::FQueryFilterData& QueryFilterData, ICollisionQueryFilterCallbackBase& QueryCallback, const ChaosInterface::FQueryDebugParams& DebugParams)
+{
+	using namespace Chaos;
+	using namespace ChaosInterface;
+
+	const FAABB3 Bounds = QueryGeom.CalculateTransformedBounds(GeomPose);
+
+	HitBuffer.IncFlushCount();
+
+	bool bSkipNarrowPhase = QueryFilterData.flags & FPhysicsQueryFlag::eSKIPNARROWPHASE;
+
+	constexpr bool bGTData = std::is_same<TOverlapHit, FOverlapHit>::value;
+	if (bSkipNarrowPhase)
+	{
+		TBPVisitor<QueryGeomType, TPayload, TOverlapHit, bGTData> OverlapVisitor(GeomPose, HitBuffer, QueryFilterData, QueryCallback, QueryGeom, DebugParams);
+		SpatialAcceleration.Overlap(Bounds, OverlapVisitor);
+	}
+	else
+	{
+		TSQVisitor<QueryGeomType, TPayload, TOverlapHit, bGTData> OverlapVisitor(GeomPose, HitBuffer, QueryFilterData, QueryCallback, QueryGeom, DebugParams);
+		SpatialAcceleration.Overlap(Bounds, OverlapVisitor);
+	}
+	HitBuffer.DecFlushCount();
+}
