[Engine]

diff --git a/Engine/Source/Runtime/PacketHandlers/PacketHandler/Public/PacketHandler.h b/Engine/Source/Runtime/PacketHandlers/PacketHandler/Public/PacketHandler.h
--- a/Engine/Source/Runtime/PacketHandlers/PacketHandler/Public/PacketHandler.h
+++ b/Engine/Source/Runtime/PacketHandlers/PacketHandler/Public/PacketHandler.h
@@ -119,7 +119,7 @@ namespace UE
 /**
  * The result of calling Incoming and Outgoing in the PacketHandler
  */
-struct PACKETHANDLER_API ProcessedPacket
+struct ProcessedPacket
 {
 	/** Pointer to the returned packet data */
 	uint8* Data;
@@ -152,7 +152,7 @@ public:
 /**
  * PacketHandler will buffer packets, this struct is used to buffer such packets while handler components are initialized
  */
-struct PACKETHANDLER_API BufferedPacket
+struct BufferedPacket
 {
 	/** Buffered packet data */
 	uint8* Data;
@@ -214,7 +214,7 @@ public:
 	/**
 	 * Base destructor
 	 */
-	~BufferedPacket();
+	PACKETHANDLER_API ~BufferedPacket();
 
 	void CountBytes(FArchive& Ar) const
 	{
@@ -249,7 +249,7 @@ enum class EIncomingResult : uint8
 /**
  * This class maintains an array of all PacketHandler Components and forwards incoming and outgoing packets the each component
  */
-class PACKETHANDLER_API PacketHandler
+class PacketHandler
 {
 public:
 	/**
@@ -257,7 +257,7 @@ public:
 	 *
 	 * @param InDDoS			Reference to the owning net drivers DDoS detection handler
 	 */
-	PacketHandler(FDDoSDetection* InDDoS=nullptr);
+	PACKETHANDLER_API PacketHandler(FDDoSDetection* InDDoS=nullptr);
 
 	virtual ~PacketHandler() = default;
 
@@ -271,7 +271,7 @@ public:
 	 * @param InDDoS				Reference to the owning net drivers DDoS detection handler
 	 * @param InDriverProfile		The PacketHandler configuration profile to use
 	 */
-	void Initialize(UE::Handler::Mode Mode, uint32 InMaxPacketBits, bool bConnectionlessOnly=false,
+	PACKETHANDLER_API void Initialize(UE::Handler::Mode Mode, uint32 InMaxPacketBits, bool bConnectionlessOnly=false,
 					TSharedPtr<class IAnalyticsProvider> InProvider=nullptr, FDDoSDetection* InDDoS=nullptr, FName InDriverProfile=NAME_None);
 
 	/**
@@ -280,7 +280,7 @@ public:
 	 * @param InLowLevelSendDel		The delegate the PacketHandler should use for triggering packet sends
 	 * @param InAddHandlerDel		Callback for notifying of new HandlerComponent's
 	 */
-	void InitializeDelegates(FPacketHandlerLowLevelSendTraits InLowLevelSendDel,
+	PACKETHANDLER_API void InitializeDelegates(FPacketHandlerLowLevelSendTraits InLowLevelSendDel,
 								FPacketHandlerNotifyAddHandler InAddHandlerDel=FPacketHandlerNotifyAddHandler());
 
 	/**
@@ -288,7 +288,7 @@ public:
 	 *
 	 * @param InFaultRecovery	A reference to the fault recovery interface
 	 */
-	void InitFaultRecovery(UE::Net::FNetConnectionFaultRecoveryBase* InFaultRecovery);
+	PACKETHANDLER_API void InitFaultRecovery(UE::Net::FNetConnectionFaultRecoveryBase* InFaultRecovery);
 
 	/**
 	 * Notification that the NetDriver analytics provider has been updated (NOT called on first initialization)
@@ -297,21 +297,21 @@ public:
 	 * @param InProvider		The analytics provider
 	 * @param InAggregator		The net analytics aggregator
 	 */
-	void NotifyAnalyticsProvider(TSharedPtr<IAnalyticsProvider> InProvider, TSharedPtr<FNetAnalyticsAggregator> InAggregator);
+	PACKETHANDLER_API void NotifyAnalyticsProvider(TSharedPtr<IAnalyticsProvider> InProvider, TSharedPtr<FNetAnalyticsAggregator> InAggregator);
 
 	/**
 	 * Triggers initialization of HandlerComponents.
 	 */
-	void InitializeComponents();
+	PACKETHANDLER_API void InitializeComponents();
 
 
 	/**
 	 * Triggered by the higher level netcode, to begin any required HandlerComponent handshakes
 	 */
-	void BeginHandshaking(FPacketHandlerHandshakeComplete InHandshakeDel=FPacketHandlerHandshakeComplete());
+	PACKETHANDLER_API void BeginHandshaking(FPacketHandlerHandshakeComplete InHandshakeDel=FPacketHandlerHandshakeComplete());
 
 
-	void Tick(float DeltaTime);
+	PACKETHANDLER_API void Tick(float DeltaTime);
 
 	/**
 	 * Adds a HandlerComponent to the pipeline, prior to initialization (none can be added after initialization)
@@ -319,7 +319,7 @@ public:
 	 * @param NewHandler		The HandlerComponent to add
 	 * @param bDeferInitialize	Whether or not to defer triggering Initialize (for batch-adds - code calling this, triggers it instead)
 	 */
-	void AddHandler(TSharedPtr<HandlerComponent>& NewHandler, bool bDeferInitialize=false);
+	PACKETHANDLER_API void AddHandler(TSharedPtr<HandlerComponent>& NewHandler, bool bDeferInitialize=false);
 
 	/**
 	 * As above, but initializes from a string specifying the component module, and (optionally) additional options
@@ -327,7 +327,7 @@ public:
 	 * @param ComponentStr		The handler component to load
 	 * @param bDeferInitialize	Whether or not to defer triggering Initialize (for batch-adds - code calling this, triggers it instead)
 	 */
-	TSharedPtr<HandlerComponent> AddHandler(const FString& ComponentStr, bool bDeferInitialize=false);
+	PACKETHANDLER_API TSharedPtr<HandlerComponent> AddHandler(const FString& ComponentStr, bool bDeferInitialize=false);
 
 
 	// @todo #JohnB: Add runtime-calculated arrays for each packet pipeline type, to reduce redundant iterations,
@@ -346,7 +346,7 @@ public:
 	 *
 	 * @param Reader	The FBitReader for the incoming packet
 	 */
-	void IncomingHigh(FBitReader& Reader);
+	PACKETHANDLER_API void IncomingHigh(FBitReader& Reader);
 
 	/**
 	 * @todo #JohnB: Work in progress, don't use yet.
@@ -357,7 +357,7 @@ public:
 	 *
 	 * @param Writer	The FBitWriter for the outgoing packet
 	 */
-	void OutgoingHigh(FBitWriter& Writer);
+	PACKETHANDLER_API void OutgoingHigh(FBitWriter& Writer);
 
 	// @todo: Don't deprecate, until after the NetDriver refactor
 	//UE_DEPRECATED(4.26, "Incoming now uses FReceivedPacketView.")
@@ -441,12 +441,12 @@ public:
 	}
 
 	/** Returns a pointer to the component set as the encryption handler, if any. */
-	TSharedPtr<FEncryptionComponent> GetEncryptionComponent();
+	PACKETHANDLER_API TSharedPtr<FEncryptionComponent> GetEncryptionComponent();
 
 	/** Returns a pointer to the first component in the HandlerComponents array with the specified name. */
-	TSharedPtr<HandlerComponent> GetComponentByName(FName ComponentName) const;
+	PACKETHANDLER_API TSharedPtr<HandlerComponent> GetComponentByName(FName ComponentName) const;
 
-	virtual void CountBytes(FArchive& Ar) const;
+	PACKETHANDLER_API virtual void CountBytes(FArchive& Ar) const;
 
 protected:
 	UE_DEPRECATED(4.26, "Incoming_Internal now uses FReceivedPacketView")
@@ -476,7 +476,7 @@ protected:
 	 * @param PacketView	View of the packet being processed - PacketView.Data should be used to return the result
 	 * @return				Returns Success/Failure
 	 */
-	EIncomingResult Incoming_Internal(FReceivedPacketView& PacketView);
+	PACKETHANDLER_API EIncomingResult Incoming_Internal(FReceivedPacketView& PacketView);
 
 	/**
 	 * Internal handling for Outgoing/OutgoingConnectionless
@@ -488,7 +488,7 @@ protected:
 	 * @param Address			The address the packet is being sent to
 	 * @return					Returns the final packet
 	 */
-	const ProcessedPacket Outgoing_Internal(uint8* Packet, int32 CountBits, FOutPacketTraits& Traits, bool bConnectionless, const TSharedPtr<const FInternetAddr>& Address);
+	PACKETHANDLER_API const ProcessedPacket Outgoing_Internal(uint8* Packet, int32 CountBits, FOutPacketTraits& Traits, bool bConnectionless, const TSharedPtr<const FInternetAddr>& Address);
 
 public:
 
@@ -501,13 +501,13 @@ public:
 	 * @param Writer		The packet being sent
 	 * @param Traits		The traits applied to the packet, if applicable
 	 */
-	void SendHandlerPacket(HandlerComponent* InComponent, FBitWriter& Writer, FOutPacketTraits& Traits);
+	PACKETHANDLER_API void SendHandlerPacket(HandlerComponent* InComponent, FBitWriter& Writer, FOutPacketTraits& Traits);
 
 
 	/**
 	 * Triggered when a child HandlerComponent has been initialized
 	 */
-	void HandlerComponentInitialized(HandlerComponent* InComponent);
+	PACKETHANDLER_API void HandlerComponentInitialized(HandlerComponent* InComponent);
 
 	/**
 	 * Queue's a packet to be sent when the handler is ticked (as a raw packet, since it's already been processed)
@@ -535,7 +535,7 @@ public:
 	 * @param InComponentName	The PacketHandler Component to search for
 	 * @return if there is a profile that has the component included.
 	 */
-	static bool DoesAnyProfileHaveComponent(const FString& InComponentName);
+	static PACKETHANDLER_API bool DoesAnyProfileHaveComponent(const FString& InComponentName);
 
 	/**
 	 * Searches the PacketHandler profile configuration for the given netdriver to find if a component is listed.
@@ -544,35 +544,35 @@ public:
 	 * @param InComponentName	The component to search for
 	 * @return if the component is listed in the profile configuration.
 	 */
-	static bool DoesProfileHaveComponent(const FName InNetDriverName, const FString& InComponentName);
+	static PACKETHANDLER_API bool DoesProfileHaveComponent(const FName InNetDriverName, const FString& InComponentName);
 
 	/**
 	 * Gets a packet from the buffered packet queue for sending
 	 *
 	 * @return		The packet to be sent, or nullptr if none are to be sent
 	 */
-	BufferedPacket* GetQueuedPacket();
+	PACKETHANDLER_API BufferedPacket* GetQueuedPacket();
 
 	/**
 	* Gets a packet from the buffered packet queue for sending (as a raw packet)
 	*
 	* @return		The packet to be sent, or nullptr if none are to be sent
 	*/
-	BufferedPacket* GetQueuedRawPacket();
+	PACKETHANDLER_API BufferedPacket* GetQueuedRawPacket();
 
 	/**
 	 * Gets a packet from the buffered connectionless packet queue for sending
 	 *
 	 * @return		The packet to be sent, or nullptr if none are to be sent
 	 */
-	BufferedPacket* GetQueuedConnectionlessPacket();
+	PACKETHANDLER_API BufferedPacket* GetQueuedConnectionlessPacket();
 
 	/**
 	 * Gets the combined reserved packet/protocol bits from all handlers, for reserving space in the parent connections packets
 	 *
 	 * @return	The combined reserved packet/protocol bits
 	 */
-	int32 GetTotalReservedPacketBits();
+	PACKETHANDLER_API int32 GetTotalReservedPacketBits();
 
 
 	/**
@@ -618,39 +618,39 @@ private:
 	 *
 	 * @param InState	The new state for the handler
 	 */
-	void SetState(UE::Handler::State InState);
+	PACKETHANDLER_API void SetState(UE::Handler::State InState);
 
 	/**
 	 * Called when net send/receive functions are triggered, when the handler is still uninitialized - to set a valid initial state
 	 */
-	void UpdateInitialState();
+	PACKETHANDLER_API void UpdateInitialState();
 
 	/**
 	 * Called when handler is finished initializing
 	 */
-	void HandlerInitialized();
+	PACKETHANDLER_API void HandlerInitialized();
 
 	/**
 	 * Replaces IncomingPacket with all unread data from ReplacementPacket
 	 *
 	 * @param ReplacementPacket		The packet whose unread data should replace IncomingPacket
 	 */
-	void ReplaceIncomingPacket(FBitReader& ReplacementPacket);
+	PACKETHANDLER_API void ReplaceIncomingPacket(FBitReader& ReplacementPacket);
 
 	/**
 	 * Takes a Packet whose position is not at bit 0, and shifts/aligns all packet data to place the current bit at position 0
 	 *
 	 * @param Packet	The packet to realign
 	 */
-	void RealignPacket(FBitReader& Packet);
+	PACKETHANDLER_API void RealignPacket(FBitReader& Packet);
 
 
 public:
 	/** Mode of the handler, Client or Server */
 	UE::Handler::Mode Mode;
 
-	static FPacketHandlerAddComponentByNameDelegate& GetAddComponentByNameDelegate();
-	static FPacketHandlerAddComponentDelegate& GetAddComponentDelegate();
+	static PACKETHANDLER_API FPacketHandlerAddComponentByNameDelegate& GetAddComponentByNameDelegate();
+	static PACKETHANDLER_API FPacketHandlerAddComponentDelegate& GetAddComponentDelegate();
 
 private:
 	/** Whether or not this PacketHandler handles connectionless (i.e. non-UNetConnection) data */
@@ -723,7 +723,7 @@ private:
 /**
  * This class appends or modifies incoming and outgoing packets on a connection
  */
-class PACKETHANDLER_API HandlerComponent
+class HandlerComponent
 {
 	friend class PacketHandler;
 
@@ -731,12 +731,12 @@ public:
 	/**
 	 * Base constructor
 	 */
-	HandlerComponent();
+	PACKETHANDLER_API HandlerComponent();
 
 	/**
 	 * Constructor that accepts a name
 	 */
-	explicit HandlerComponent(FName InName);
+	PACKETHANDLER_API explicit HandlerComponent(FName InName);
 
 	/**
 	 * Base destructor
@@ -748,7 +748,7 @@ public:
 	/**
 	 * Returns whether this handler is currently active
 	 */
-	virtual bool IsActive() const;
+	PACKETHANDLER_API virtual bool IsActive() const;
 
 	/**
 	 * Return whether this handler is valid
@@ -758,7 +758,7 @@ public:
 	/**
 	 * Returns whether this handler is initialized
 	 */
-	bool IsInitialized() const;
+	PACKETHANDLER_API bool IsInitialized() const;
 
 	/**
 	* Returns whether this handler perform a network handshake during initialization
@@ -869,7 +869,7 @@ public:
 	 *
 	 * @param Active	Whether or not the handled should be active
 	 */
-	virtual void SetActive(bool Active);
+	PACKETHANDLER_API virtual void SetActive(bool Active);
 
 	/**
 	 * Returns the amount of reserved packet/protocol bits expected from this component.
@@ -890,7 +890,7 @@ public:
 	 */
 	virtual void NotifyAnalyticsProvider() {}
 
-	virtual void CountBytes(FArchive& Ar) const;
+	PACKETHANDLER_API virtual void CountBytes(FArchive& Ar) const;
 
 protected:
 	/**
@@ -898,12 +898,12 @@ protected:
 	 *
 	 * @param State		The new state for the handler
 	 */
-	void SetState(UE::Handler::Component::State State);
+	PACKETHANDLER_API void SetState(UE::Handler::Component::State State);
 
 	/**
 	 * Should be called when the handler is fully initialized on both remote and local
 	 */
-	void Initialized();
+	PACKETHANDLER_API void Initialized();
 
 
 public:
@@ -937,7 +937,7 @@ private:
 /**
  * PacketHandler Module Interface
  */
-class PACKETHANDLER_API FPacketHandlerComponentModuleInterface : public IModuleInterface
+class FPacketHandlerComponentModuleInterface : public IModuleInterface
 {
 public:
 	/* Creates an instance of this component */
@@ -946,7 +946,7 @@ public:
 		return nullptr;
 	}
 
-	virtual void StartupModule() override;
+	PACKETHANDLER_API virtual void StartupModule() override;
 
-	virtual void ShutdownModule() override;
+	PACKETHANDLER_API virtual void ShutdownModule() override;
 };
diff --git a/Engine/Source/Runtime/PakFile/Internal/IoDispatcherFileBackendTypes.h b/Engine/Source/Runtime/PakFile/Internal/IoDispatcherFileBackendTypes.h
--- a/Engine/Source/Runtime/PakFile/Internal/IoDispatcherFileBackendTypes.h
+++ b/Engine/Source/Runtime/PakFile/Internal/IoDispatcherFileBackendTypes.h
@@ -23,7 +23,7 @@
 
 struct FFileIoStoreCompressionContext;
 
-struct PAKFILE_API FFileIoStoreContainerFilePartition
+struct FFileIoStoreContainerFilePartition
 {
 	FFileIoStoreContainerFilePartition() = default;
 	FFileIoStoreContainerFilePartition(FFileIoStoreContainerFilePartition&&) = delete;
@@ -40,7 +40,7 @@ struct PAKFILE_API FFileIoStoreContainerFilePartition
 	std::atomic<int32> StartedReadRequestsCount = 0;
 };
 
-struct PAKFILE_API FFileIoStoreBlockSignatureTable
+struct FFileIoStoreBlockSignatureTable
 {
 	TArray<FSHAHash> Hashes;
 
@@ -61,7 +61,7 @@ private:
 	std::atomic<int64> RefCount = 0;
 };
 
-struct PAKFILE_API FFileIoStoreContainerFile
+struct FFileIoStoreContainerFile
 {
 	FFileIoStoreContainerFile() = default;
 	FFileIoStoreContainerFile(FFileIoStoreContainerFile&&) = default;
@@ -90,13 +90,13 @@ struct PAKFILE_API FFileIoStoreContainerFile
 	}
 };
 
-struct PAKFILE_API FFileIoStoreBuffer
+struct FFileIoStoreBuffer
 {
 	FFileIoStoreBuffer* Next = nullptr;
 	uint8* Memory = nullptr;
 };
 
-struct PAKFILE_API FFileIoStoreBlockKey
+struct FFileIoStoreBlockKey
 {
 	union
 	{
@@ -120,7 +120,7 @@ struct PAKFILE_API FFileIoStoreBlockKey
 	}
 };
 
-struct PAKFILE_API FFileIoStoreBlockScatter
+struct FFileIoStoreBlockScatter
 {
 	struct FFileIoStoreResolvedRequest* Request = nullptr;
 	uint64 DstOffset = 0;
@@ -128,7 +128,7 @@ struct PAKFILE_API FFileIoStoreBlockScatter
 	uint64 Size = 0;
 };
 
-struct PAKFILE_API FFileIoStoreCompressedBlock
+struct FFileIoStoreCompressedBlock
 {
 	FFileIoStoreCompressedBlock* Next = nullptr;
 	FFileIoStoreBlockKey Key;
@@ -150,7 +150,7 @@ struct PAKFILE_API FFileIoStoreCompressedBlock
 	bool bCancelled = false;
 };
 
-struct PAKFILE_API FFileIoStoreReadRequest
+struct FFileIoStoreReadRequest
 {
 	enum EQueueStatus
 	{
@@ -191,7 +191,7 @@ struct PAKFILE_API FFileIoStoreReadRequest
 #endif
 
 private:
-	static uint32 NextSequence;
+	static PAKFILE_API uint32 NextSequence;
 };
 
 #define CHECK_IO_STORE_READ_REQUEST_LIST_MEMBERSHIP (DO_CHECK && 0)
@@ -275,7 +275,7 @@ private:
 	FFileIoStoreReadRequest* Next = nullptr;;
 };
 
-class PAKFILE_API FFileIoStoreReadRequestListIterator
+class FFileIoStoreReadRequestListIterator
 {
 public:
 	FFileIoStoreReadRequestListIterator(const FFileIoStoreReadRequestListIterator&) = default;
@@ -337,7 +337,7 @@ private:
 
 // Wrapper for doubly-linked intrusive list of FFileIoStoreReadRequest
 
-class PAKFILE_API FFileIoStoreReadRequestList
+class FFileIoStoreReadRequestList
 {
 public:
 	FFileIoStoreReadRequestList()
@@ -525,7 +525,7 @@ private:
 
 #if CHECK_IO_STORE_READ_REQUEST_LIST_MEMBERSHIP
 	uint32 ListCookie;
-	static uint32 NextListCookie;
+	static PAKFILE_API uint32 NextListCookie;
 #endif
 	
 	void AppendSteal(FFileIoStoreReadRequest* ListHead, FFileIoStoreReadRequest* ListTail)
@@ -561,7 +561,7 @@ private:
 
 class FFileIoStoreStats;
 
-class PAKFILE_API FFileIoStoreBufferAllocator
+class FFileIoStoreBufferAllocator
 {
 public:
 	FFileIoStoreBufferAllocator(FFileIoStoreStats& InStats)
@@ -569,9 +569,9 @@ public:
 	{
 	}
 
-	void Initialize(uint64 MemorySize, uint64 BufferSize, uint32 BufferAlignment);
-	FFileIoStoreBuffer* AllocBuffer();
-	void FreeBuffer(FFileIoStoreBuffer* Buffer);
+	PAKFILE_API void Initialize(uint64 MemorySize, uint64 BufferSize, uint32 BufferAlignment);
+	PAKFILE_API FFileIoStoreBuffer* AllocBuffer();
+	PAKFILE_API void FreeBuffer(FFileIoStoreBuffer* Buffer);
 	uint64 GetBufferSize() const { return BufferSize; }
 
 private:
@@ -582,15 +582,15 @@ private:
 	FFileIoStoreBuffer* FirstFreeBuffer = nullptr;
 };
 
-class PAKFILE_API FFileIoStoreBlockCache
+class FFileIoStoreBlockCache
 {
 public:
-	FFileIoStoreBlockCache(FFileIoStoreStats& Stats);
-	~FFileIoStoreBlockCache();
+	PAKFILE_API FFileIoStoreBlockCache(FFileIoStoreStats& Stats);
+	PAKFILE_API ~FFileIoStoreBlockCache();
 
-	void Initialize(uint64 CacheMemorySize, uint64 ReadBufferSize);
-	bool Read(FFileIoStoreReadRequest* Block);
-	void Store(const FFileIoStoreReadRequest* Block);
+	PAKFILE_API void Initialize(uint64 CacheMemorySize, uint64 ReadBufferSize);
+	PAKFILE_API bool Read(FFileIoStoreReadRequest* Block);
+	PAKFILE_API void Store(const FFileIoStoreReadRequest* Block);
 
 private:
 	struct FCachedBlock
@@ -609,7 +609,7 @@ private:
 	uint64 ReadBufferSize = 0;
 };
 
-struct PAKFILE_API FFileIoStoreReadRequestSortKey
+struct FFileIoStoreReadRequestSortKey
 {
 	uint64 Offset = 0;
 	uint64 Handle = 0;
@@ -623,10 +623,10 @@ struct PAKFILE_API FFileIoStoreReadRequestSortKey
 };
 
 // Stores FFileIoStoreReadRequest sorted by file handle & offset with a parallel list sorted by insertion order 
-class PAKFILE_API FFileIoStoreOffsetSortedRequestQueue
+class FFileIoStoreOffsetSortedRequestQueue
 {
 public:
-	FFileIoStoreOffsetSortedRequestQueue(int32 InPriority);
+	PAKFILE_API FFileIoStoreOffsetSortedRequestQueue(int32 InPriority);
 	FFileIoStoreOffsetSortedRequestQueue(const FFileIoStoreOffsetSortedRequestQueue&) = delete;
 	FFileIoStoreOffsetSortedRequestQueue(FFileIoStoreOffsetSortedRequestQueue&&) = default;
 	FFileIoStoreOffsetSortedRequestQueue& operator=(const FFileIoStoreOffsetSortedRequestQueue&) = delete;
@@ -636,13 +636,13 @@ public:
 	bool IsEmpty() const { return Requests.Num() == 0; }
 	
 	// Remove all requests from this container and return them, for switching to a different queue scheme
-	TArray<FFileIoStoreReadRequest*> StealRequests();
+	PAKFILE_API TArray<FFileIoStoreReadRequest*> StealRequests();
 	// Remove all requests whose priority has been changed to something other than the Priority of this queue
-	TArray<FFileIoStoreReadRequest*> RemoveMisprioritizedRequests();
+	PAKFILE_API TArray<FFileIoStoreReadRequest*> RemoveMisprioritizedRequests();
 
-	FFileIoStoreReadRequest* Pop(FFileIoStoreReadRequestSortKey LastSortKey);
-	void Push(FFileIoStoreReadRequest* Request);
-	int32 HandleContainerUnmounted(const FFileIoStoreContainerFile& ContainerFile);
+	PAKFILE_API FFileIoStoreReadRequest* Pop(FFileIoStoreReadRequestSortKey LastSortKey);
+	PAKFILE_API void Push(FFileIoStoreReadRequest* Request);
+	PAKFILE_API int32 HandleContainerUnmounted(const FFileIoStoreContainerFile& ContainerFile);
 
 private:
 	int32 Priority;
@@ -655,22 +655,22 @@ private:
 	// We store this on the heap in case we get moved, FFileIoStoreReadRequest keeps pointers to FFileIoStoreReadRequestList for debugging
 	FFileIoStoreReadRequestList RequestsBySequence;
 
-	FFileIoStoreReadRequest* GetNextInternal(FFileIoStoreReadRequestSortKey LastSortKey, bool bPop);
+	PAKFILE_API FFileIoStoreReadRequest* GetNextInternal(FFileIoStoreReadRequestSortKey LastSortKey, bool bPop);
 
 	static FFileIoStoreReadRequestSortKey RequestSortProjection(FFileIoStoreReadRequest* Request) { return FFileIoStoreReadRequestSortKey(Request); }
-	static bool RequestSortPredicate(const FFileIoStoreReadRequestSortKey& A, const FFileIoStoreReadRequestSortKey& B);
+	static PAKFILE_API bool RequestSortPredicate(const FFileIoStoreReadRequestSortKey& A, const FFileIoStoreReadRequestSortKey& B);
 };
 
-class PAKFILE_API FFileIoStoreRequestQueue
+class FFileIoStoreRequestQueue
 {
 public:
-	FFileIoStoreReadRequest* Pop();
-	void Push(FFileIoStoreReadRequest& Request);	// Takes ownership of Request and rewrites its intrustive linked list pointers
-	void Push(FFileIoStoreReadRequestList& Requests); // Consumes the request list and overwrites all intrustive linked list pointers
-	void UpdateOrder();
-	void Lock();
-	void Unlock();
-	int32 HandleContainerUnmounted(const FFileIoStoreContainerFile& ContainerFile);
+	PAKFILE_API FFileIoStoreReadRequest* Pop();
+	PAKFILE_API void Push(FFileIoStoreReadRequest& Request);	// Takes ownership of Request and rewrites its intrustive linked list pointers
+	PAKFILE_API void Push(FFileIoStoreReadRequestList& Requests); // Consumes the request list and overwrites all intrustive linked list pointers
+	PAKFILE_API void UpdateOrder();
+	PAKFILE_API void Lock();
+	PAKFILE_API void Unlock();
+	PAKFILE_API int32 HandleContainerUnmounted(const FFileIoStoreContainerFile& ContainerFile);
 
 private:
 	static bool QueueSortFunc(const FFileIoStoreReadRequest& A, const FFileIoStoreReadRequest& B)
@@ -704,7 +704,7 @@ private:
 template <typename T, uint16 SlabSize = 4096>
 using TIoDispatcherSingleThreadedSlabAllocator = TSingleThreadedSlabAllocator<T, SlabSize>;
 
-struct PAKFILE_API FFileIoStoreReadRequestLink
+struct FFileIoStoreReadRequestLink
 {
 	FFileIoStoreReadRequestLink(FFileIoStoreReadRequest& InReadRequest)
 		: ReadRequest(InReadRequest)
@@ -715,7 +715,7 @@ struct PAKFILE_API FFileIoStoreReadRequestLink
 	FFileIoStoreReadRequest& ReadRequest;
 };
 
-class PAKFILE_API FFileIoStoreRequestAllocator
+class FFileIoStoreRequestAllocator
 {
 public:
 	int64 GetLiveReadRequestsCount() const
@@ -789,10 +789,10 @@ private:
 	int64 LiveReadRequestsCount = 0;
 };
 
-struct PAKFILE_API FFileIoStoreResolvedRequest
+struct FFileIoStoreResolvedRequest
 {
 public:
-	FFileIoStoreResolvedRequest(
+	PAKFILE_API FFileIoStoreResolvedRequest(
 		FIoRequestImpl& InDispatcherRequest,
 		FFileIoStoreContainerFile* InContainerFile,
 		uint64 InResolvedOffset,
@@ -842,7 +842,7 @@ public:
 		return DispatcherRequest->GetBuffer();
 	}
 
-	void AddReadRequestLink(FFileIoStoreReadRequestLink* ReadRequestLink);
+	PAKFILE_API void AddReadRequestLink(FFileIoStoreReadRequestLink* ReadRequestLink);
 
 private:
 	FIoRequestImpl* DispatcherRequest = nullptr;
@@ -863,18 +863,18 @@ private:
 
 // Wrapper for sending stats to both insights and csv profiling from any platform-specific dispatcher.
 #if UE_FILEIOSTORE_STATS_ENABLED
-class PAKFILE_API FFileIoStoreStats
+class FFileIoStoreStats
 {
 public:
-	FFileIoStoreStats();
-	~FFileIoStoreStats();
+	PAKFILE_API FFileIoStoreStats();
+	PAKFILE_API ~FFileIoStoreStats();
 
 	// Called by the backend when underlying file system reads start
-	void OnFilesystemReadStarted(const FFileIoStoreReadRequest* Request);
-	void OnFilesystemReadsStarted(const FFileIoStoreReadRequestList& Requests);
+	PAKFILE_API void OnFilesystemReadStarted(const FFileIoStoreReadRequest* Request);
+	PAKFILE_API void OnFilesystemReadsStarted(const FFileIoStoreReadRequestList& Requests);
 	// Called by the backend when underlying filesystem reads complete, possibly already decompressed on some systems
-	void OnFilesystemReadCompleted(const FFileIoStoreReadRequest* Request);
-	void OnFilesystemReadsCompleted(const FFileIoStoreReadRequestList& CompletedRequests);
+	PAKFILE_API void OnFilesystemReadCompleted(const FFileIoStoreReadRequest* Request);
+	PAKFILE_API void OnFilesystemReadsCompleted(const FFileIoStoreReadRequestList& CompletedRequests);
 
 private:
 	friend class FFileIoStore;
@@ -927,35 +927,35 @@ private:
 		return float(double(Bytes) / 1024.0);
 	}
 
-	bool CsvTick(float DeltaTime);
+	PAKFILE_API bool CsvTick(float DeltaTime);
 
-	void OnReadRequestsQueued(const FFileIoStoreReadRequestList& Requests);
-	void OnReadRequestsCompleted(const FFileIoStoreReadRequestList& Requests);
-	void OnDecompressQueued(const FFileIoStoreCompressedBlock* CompressedBlock);
-	void OnDecompressComplete(const FFileIoStoreCompressedBlock* CompressedBlock);
+	PAKFILE_API void OnReadRequestsQueued(const FFileIoStoreReadRequestList& Requests);
+	PAKFILE_API void OnReadRequestsCompleted(const FFileIoStoreReadRequestList& Requests);
+	PAKFILE_API void OnDecompressQueued(const FFileIoStoreCompressedBlock* CompressedBlock);
+	PAKFILE_API void OnDecompressComplete(const FFileIoStoreCompressedBlock* CompressedBlock);
 
 	// Bytes were copied into the buffer provided to users of the io system
-	void OnBytesScattered(int64 BytesScattered);
+	PAKFILE_API void OnBytesScattered(int64 BytesScattered);
 	// Record stats for block cache hit/miss rate & data throughput
-	void OnBlockCacheStore(uint64 NumBytes);
-	void OnBlockCacheHit(uint64 NumBytes);
-	void OnBlockCacheMiss(uint64 NumBytes);
+	PAKFILE_API void OnBlockCacheStore(uint64 NumBytes);
+	PAKFILE_API void OnBlockCacheHit(uint64 NumBytes);
+	PAKFILE_API void OnBlockCacheMiss(uint64 NumBytes);
 
 	// A read was started without seeking
-	void OnSequentialRead();
+	PAKFILE_API void OnSequentialRead();
 	// A read was started and was either a forward or reverse seek
-	void OnSeek(uint64 LastOffset, uint64 NewOffset);
+	PAKFILE_API void OnSeek(uint64 LastOffset, uint64 NewOffset);
 	// A read was started from a different file handle from the last one we used
-	void OnHandleChangeSeek();
+	PAKFILE_API void OnHandleChangeSeek();
 
-	void OnTocMounted(uint64 AllocatedSize);
-	void OnTocUnmounted(uint64 AllocatedSize);
+	PAKFILE_API void OnTocMounted(uint64 AllocatedSize);
+	PAKFILE_API void OnTocUnmounted(uint64 AllocatedSize);
 
-	void OnBufferReleased();
-	void OnBufferAllocated();
+	PAKFILE_API void OnBufferReleased();
+	PAKFILE_API void OnBufferAllocated();
 };
 #else
-class PAKFILE_API FFileIoStoreStats
+class FFileIoStoreStats
 {
 public:
 	void OnFilesystemReadStarted(const FFileIoStoreReadRequest* Request) {}
@@ -984,7 +984,7 @@ private:
 };
 #endif
 
-struct PAKFILE_API FInitializePlatformFileIoStoreParams
+struct FInitializePlatformFileIoStoreParams
 {
 	const FWakeUpIoDispatcherThreadDelegate* WakeUpDispatcherThreadDelegate = nullptr;
 	FFileIoStoreRequestAllocator* RequestAllocator = nullptr;
diff --git a/Engine/Source/Runtime/PakFile/Public/IPlatformFilePak.h b/Engine/Source/Runtime/PakFile/Public/IPlatformFilePak.h
--- a/Engine/Source/Runtime/PakFile/Public/IPlatformFilePak.h
+++ b/Engine/Source/Runtime/PakFile/Public/IPlatformFilePak.h
@@ -678,22 +678,22 @@ FORCEINLINE FArchive& operator<<(FArchive& Ar, FPakEntryLocation& PakEntryLocati
 class FPakFile;
 
 // Wrapper for a pointer to a shared pak reader archive that has been temporarily acquired. 
-class PAKFILE_API FSharedPakReader final 
+class FSharedPakReader final 
 {
 	friend class FPakFile;
 
 	FArchive* Archive = nullptr;
 	FPakFile* PakFile = nullptr; // Pak file to return ownership to on destruction
 
-	FSharedPakReader(FArchive* InArchive, FPakFile* InPakFile);
+	PAKFILE_API FSharedPakReader(FArchive* InArchive, FPakFile* InPakFile);
 
 public:
-	~FSharedPakReader();
+	PAKFILE_API ~FSharedPakReader();
 
 	FSharedPakReader(const FSharedPakReader& Other) = delete;
 	FSharedPakReader& operator=(const FSharedPakReader& Other) = delete;
-	FSharedPakReader(FSharedPakReader&& Other);
-	FSharedPakReader& operator=(FSharedPakReader&& Other);
+	PAKFILE_API FSharedPakReader(FSharedPakReader&& Other);
+	PAKFILE_API FSharedPakReader& operator=(FSharedPakReader&& Other);
 
 	explicit operator bool() const { return Archive != nullptr; }
 	bool operator==(nullptr_t) { return Archive == nullptr; }
@@ -719,7 +719,7 @@ struct FPakEntryPair
 /**
  * Pak file.
  */
-class PAKFILE_API FPakFile : FNoncopyable, public FRefCountBase, public IPakFile
+class FPakFile : FNoncopyable, public FRefCountBase, public IPakFile
 {
 public:
 	/** Index data that provides a map from the hash of a Filename to an FPakEntryLocation */
@@ -764,7 +764,7 @@ public:
 	};
 
 	/** Recreates the pak reader for each thread */
-	bool RecreatePakReaders(IPlatformFile* LowerLevel);
+	PAKFILE_API bool RecreatePakReaders(IPlatformFile* LowerLevel);
 
 	struct FArchiveAndLastAccessTime 
 	{
@@ -863,8 +863,8 @@ private:
 		return 0;
 	}
 
-	FArchive* CreatePakReader(IPlatformFile* LowerLevel, const TCHAR* Filename);
-	FArchive* SetupSignedPakReader(FArchive* Reader, const TCHAR* Filename);
+	PAKFILE_API FArchive* CreatePakReader(IPlatformFile* LowerLevel, const TCHAR* Filename);
+	PAKFILE_API FArchive* SetupSignedPakReader(FArchive* Reader, const TCHAR* Filename);
 
 
 public:
@@ -914,7 +914,7 @@ public:
 	* @param Filename Pak filename.
 	* @param bIsSigned true if the pak is signed
 	*/
-	FPakFile(const TCHAR* Filename, bool bIsSigned);
+	PAKFILE_API FPakFile(const TCHAR* Filename, bool bIsSigned);
 #endif
 
 	/**
@@ -924,7 +924,7 @@ public:
 	 * @param Filename Filename.
 	 * @param bIsSigned = true if the pak is signed.
 	 */
-	FPakFile(IPlatformFile* LowerLevel, const TCHAR* Filename, bool bIsSigned, bool bLoadIndex = true);
+	PAKFILE_API FPakFile(IPlatformFile* LowerLevel, const TCHAR* Filename, bool bIsSigned, bool bLoadIndex = true);
 
 	/**
 	 * Creates a pak file using the supplied archive.
@@ -932,12 +932,12 @@ public:
 	 * @param Archive	Pointer to the archive which contains the pak file data.
 	 */
 #if WITH_EDITOR
-	FPakFile(FArchive* Archive);
+	PAKFILE_API FPakFile(FArchive* Archive);
 #endif
 
 private:
 	/** Private destructor, use AddRef/Release instead */
-	virtual ~FPakFile();
+	PAKFILE_API virtual ~FPakFile();
 	friend class FRefCountBase;
 
 public:
@@ -957,7 +957,7 @@ public:
 	 *
 	 * @return true if this pak file has passed the initial signature checking phase
 	 */
-	bool PassedSignatureChecks() const;
+	PAKFILE_API bool PassedSignatureChecks() const;
 
 	/**
 	 * Gets pak filename.
@@ -987,23 +987,23 @@ public:
 	}
 
 	/** Returns the FullPath (includes Mount) Filename found in Pruned DirectoryIndex */
-	void GetPrunedFilenames(TArray<FString>& OutFileList) const;
+	PAKFILE_API void GetPrunedFilenames(TArray<FString>& OutFileList) const;
 
 	/** Returns the RelativePathFromMount Filename for every Filename found in the Pruned DirectoryIndex that points to a PakEntry in the given Chunk */
-	void GetPrunedFilenamesInChunk(const TArray<int32>& InChunkIDs, TArray<FString>& OutFileList) const;
+	PAKFILE_API void GetPrunedFilenamesInChunk(const TArray<int32>& InChunkIDs, TArray<FString>& OutFileList) const;
 
 	/**
 	 * Gets shared pak file archive for given thread.
 	 *
 	 * @return Pointer to pak file archive used to read data from pak.
 	 */
-	FSharedPakReader GetSharedReader(IPlatformFile* LowerLevel);
+	PAKFILE_API FSharedPakReader GetSharedReader(IPlatformFile* LowerLevel);
 
 	// Return a shared pak reader. Should only be called from the FSharedPakReader's destructor.
-	void ReturnSharedReader(FArchive* SharedReader);
+	PAKFILE_API void ReturnSharedReader(FArchive* SharedReader);
 
 	// Delete all readers that haven't been used in MaxAgeSeconds.
-	void ReleaseOldReaders(double MaxAgeSeconds);
+	PAKFILE_API void ReleaseOldReaders(double MaxAgeSeconds);
 
 	/**
 	 * Finds an entry in the pak file matching the given filename.
@@ -1018,7 +1018,7 @@ public:
 		Found,
 		FoundDeleted,
 	};
-	EFindResult Find(const FString& FullPath, FPakEntry* OutEntry) const;
+	PAKFILE_API EFindResult Find(const FString& FullPath, FPakEntry* OutEntry) const;
 
 	/**
 	 * Sets the pak file mount point.
@@ -1130,7 +1130,7 @@ public:
 	 *
 	 * @return true if the pak file is valid
 	 */
-	bool Check();
+	PAKFILE_API bool Check();
 
 	/** Base functionality for iterating over the DirectoryIndex. */
 	class FBaseIterator
@@ -1423,7 +1423,7 @@ public:
 	}
 
 	/** Hash the given full-path filename using the hash function used by FPakFiles, with the given FPakFile-specific seed, with version provided for legacy pak files that used different hash function */
-	static uint64 HashPath(const TCHAR* RelativePathFromMount, uint64 Seed, int32 PakFileVersion);
+	static PAKFILE_API uint64 HashPath(const TCHAR* RelativePathFromMount, uint64 Seed, int32 PakFileVersion);
 
 	/** Read a list of (Filename, FPakEntry) pairs from a provided enumeration, attempt to encode each one,
 	  * store each one in the appropriate given encoded and/or unencoded array, and populate the given
@@ -1445,16 +1445,16 @@ public:
 	  * @param PakFileVersion Version of the pakfile containing the index, to support legacy formats
 	  */
 	typedef TFunction<FPakEntryPair & ()> ReadNextEntryFunction;
-	static void EncodePakEntriesIntoIndex(int32 InNumEntries, const ReadNextEntryFunction& InReadNextEntry, const TCHAR* InPakFilename, const FPakInfo& InPakInfo, const FString& MountPoint,
+	static PAKFILE_API void EncodePakEntriesIntoIndex(int32 InNumEntries, const ReadNextEntryFunction& InReadNextEntry, const TCHAR* InPakFilename, const FPakInfo& InPakInfo, const FString& MountPoint,
 		int32& OutNumEncodedEntries, int32& OutNumDeletedEntries, uint64* OutPathHashSeed,
 		FDirectoryIndex* OutDirectoryIndex, FPathHashIndex* OutPathHashIndex, TArray<uint8>& OutEncodedPakEntries, TArray<FPakEntry>& OutNonEncodableEntries, TMap<uint64, FString>* InOutCollisionDetection,
 		int32 PakFileVersion);
 
 	/** Lookup the FPakEntryLocation stored in the given PathHashIndex, return nullptr if not found */
-	static const FPakEntryLocation* FindLocationFromIndex(const FString& FullPath, const FString& MountPoint, const FPathHashIndex& PathHashIndex, uint64 PathHashSeed, int32 PakFileVersion);
+	static PAKFILE_API const FPakEntryLocation* FindLocationFromIndex(const FString& FullPath, const FString& MountPoint, const FPathHashIndex& PathHashIndex, uint64 PathHashSeed, int32 PakFileVersion);
 
 	/** Lookup the FPakEntryLocation stored in the given DirectoryIndex, return nullptr if not found */
-	static const FPakEntryLocation* FindLocationFromIndex(const FString& FullPath, const FString& MountPoint, const FDirectoryIndex& DirectoryIndex);
+	static PAKFILE_API const FPakEntryLocation* FindLocationFromIndex(const FString& FullPath, const FString& MountPoint, const FDirectoryIndex& DirectoryIndex);
 
 	/**
 	  * Returns the FPakEntry pointed to by the given FPakEntryLocation inside the given EncodedPakEntries or Files
@@ -1462,7 +1462,7 @@ public:
 	  * If OutEntry is non-null, populates it with a copy of the FPakEntry found, or sets it to
 	  * an FPakEntry with SetDeleteRecord(true) if not found
 	  */
-	static EFindResult GetPakEntry(const FPakEntryLocation& FPakEntryLocation, FPakEntry* OutEntry, const TArray<uint8>& EncodedPakEntries, const TArray<FPakEntry>& Files, const FPakInfo& Info);
+	static PAKFILE_API EFindResult GetPakEntry(const FPakEntryLocation& FPakEntryLocation, FPakEntry* OutEntry, const TArray<uint8>& EncodedPakEntries, const TArray<FPakEntry>& Files, const FPakInfo& Info);
 
 	/**
 	 * Given a directory index, remove entries from it that are directed by ini to not have filenames kept at runtime.
@@ -1471,7 +1471,7 @@ public:
 	 * OutDirectoryIndex - If null, InOutDirectoryIndex will have pruned entries removed.  If non-null, InOutDirectoryIndex will not be modified, and PrunedDirectoryIndex will have kept values added.
 	 * MountPoint The mount point for the pak containing the index, used to provide the fullpath for filenames in the DirectoryIndex for comparison against paths in ini
 	 */
-	static void PruneDirectoryIndex(FDirectoryIndex& InOutDirectoryIndex, FDirectoryIndex* PrunedDirectoryIndex, const FString& MountPoint);
+	static PAKFILE_API void PruneDirectoryIndex(FDirectoryIndex& InOutDirectoryIndex, FDirectoryIndex* PrunedDirectoryIndex, const FString& MountPoint);
 
 	/* Helper function to modify the given string to append '/' at the end of path to normalize directory names for hash and string compares */
 	static void MakeDirectoryFromPath(FString& Path)
@@ -1599,64 +1599,64 @@ public:
 	}
 
 	/* Returns the global,const flag for whether the current process is allowing PakFiles to keep their entire DirectoryIndex (if it exists in the PakFile on disk) rather than pruning it */
-	static bool IsPakKeepFullDirectory();
+	static PAKFILE_API bool IsPakKeepFullDirectory();
 
 	/* Returns the global,const flag for whether UnrealPak should write a copy of the full PathHashIndex and Pruned DirectoryIndex to the PakFile */
-	static bool IsPakWritePathHashIndex();
+	static PAKFILE_API bool IsPakWritePathHashIndex();
 
 	/* Returns the global,const flag for whether UnrealPak should write a copy of the full DirectoryIndex to the PakFile */
-	static bool IsPakWriteFullDirectoryIndex();
+	static PAKFILE_API bool IsPakWriteFullDirectoryIndex();
 
 private:
 
 	/**
 	 * Initializes the pak file.
 	 */
-	void Initialize(FArchive& Reader, bool bLoadIndex = true);
+	PAKFILE_API void Initialize(FArchive& Reader, bool bLoadIndex = true);
 
 	/**
 	 * Loads and initializes pak file index.
 	 */
-	void LoadIndex(FArchive& Reader);
+	PAKFILE_API void LoadIndex(FArchive& Reader);
 
 	/**
 	  * Returns the FPakEntry pointed to by the given FPakEntryLocation, forwards to the static GetPakEntry with data from *this
 	  */
-	EFindResult GetPakEntry(const FPakEntryLocation& FPakEntryLocation, FPakEntry* OutEntry) const;
+	PAKFILE_API EFindResult GetPakEntry(const FPakEntryLocation& FPakEntryLocation, FPakEntry* OutEntry) const;
 
 	/** Helper class to read IndexSettings from project delegate and commandline */
 	struct FIndexSettings;
 
-	static FIndexSettings& GetIndexSettings();
+	static PAKFILE_API FIndexSettings& GetIndexSettings();
 
 	/**
 	  * Returns the global,const flag for whether the current process should run directory queries on both the DirectoryIndex and the Pruned DirectoryIndex and log an error if they don't match.
 	  * Validation only occurs until the first call to OptimizeMemoryUsageForMountedPaks, after which the Full DirectoryIndex is dropped and there is nothing left to Validate
 	  * Has the same effect as IsPakDelayPruning, plus the addition of the error for any mismatches.
 	  */
-	static bool IsPakValidatePruning();
+	static PAKFILE_API bool IsPakValidatePruning();
 	/**
 	 * Returns the global,const flag for whether the current process should keep a copy of the Full DirectoryIndex around until OptimizeMemoryUsageForMountedPaks is called, so that systems can run
 	 * directory queries against the full index until then.
 	 * Note that validation will still occur if IsPakValidatePruning is true.
 	 */
-	static bool IsPakDelayPruning();
+	static PAKFILE_API bool IsPakDelayPruning();
 
 #if ENABLE_PAKFILE_RUNTIME_PRUNING
 	/** Global flag for whether a Pak has indicated it needs Pruning */
-	static bool bSomePakNeedsPruning;
+	static PAKFILE_API bool bSomePakNeedsPruning;
 #endif
 
 	/**
 	  * Returns whether read accesses against the DirectoryIndex need to be guarded using this->DirectoryIndexLock.
 	  * Locking is not required if the pak is not going to be pruned or already has been; the DirectoryIndex is immutable after that point, and we can get a performance benefit by skipping the lock.
 	  */
-	bool RequiresDirectoryIndexLock() const;
+	PAKFILE_API bool RequiresDirectoryIndexLock() const;
 
 	/**
 	 * Returns whether the current Process IsPakValidatePruning and this PakFile has a Full DirectoryIndex and Pruned DirectoryIndex to validate.
 	 */
-	bool ShouldValidatePrunedDirectory() const;
+	PAKFILE_API bool ShouldValidatePrunedDirectory() const;
 
 	/**
 	  * Add the given (Filename,FPakEntryLocation) value into the provided indexes
@@ -1670,26 +1670,26 @@ private:
 	  * @param InOutCollisionDetection Optional parameter to detect hash collisions.  If present, the hashed filename will be check()'d for a collision against a different filename in InOutCollisionDetection, and will be added into InOutCollisionDetection
 	  * @param PakFileVersion Version of the pakfile containing the index, to support legacy formats
 	  */
-	static void AddEntryToIndex(const FString& Filename, const FPakEntryLocation& EntryLocation, const FString& MountPoint, uint64 PathHashSeed,
+	static PAKFILE_API void AddEntryToIndex(const FString& Filename, const FPakEntryLocation& EntryLocation, const FString& MountPoint, uint64 PathHashSeed,
 		FDirectoryIndex* DirectoryIndex, FPathHashIndex* PathHashIndex, TMap<uint64, FString>* CollisionDetection, int32 PakFileVersion);
 
 	/* Encodes a pak entry as an array of bytes into the given archive.  Returns true if encoding succeeded.  If encoding did not succeed, caller will need to store the InPakEntry in an unencoded list */
-	static bool EncodePakEntry(FArchive& Ar, const FPakEntry& InPakEntry, const FPakInfo& InInfo);
+	static PAKFILE_API bool EncodePakEntry(FArchive& Ar, const FPakEntry& InPakEntry, const FPakInfo& InInfo);
 
 	/* Decodes a bit-encoded pak entry from a pointer to the start of its encoded bytes into the given OutEntry */
-	static void DecodePakEntry(const uint8* SourcePtr, FPakEntry& OutEntry, const FPakInfo& InInfo);
+	static PAKFILE_API void DecodePakEntry(const uint8* SourcePtr, FPakEntry& OutEntry, const FPakInfo& InInfo);
 
 	/* Internal index loading function that returns false if index loading fails due to an intermittent IO error. Allows LoadIndex to retry or throw a fatal as required */
-	bool LoadIndexInternal(FArchive& Reader);
+	PAKFILE_API bool LoadIndexInternal(FArchive& Reader);
 
 	/* Legacy index loading function for PakFiles saved before FPakInfo::PakFile_Version_PathHashIndex */
-	bool LoadLegacyIndex(FArchive& Reader);
+	PAKFILE_API bool LoadLegacyIndex(FArchive& Reader);
 
 	/* Helper function for LoadIndexInternal; each array of Index bytes read from the file needs to be independently decrypted and checked for corruption */
-	bool DecryptAndValidateIndex(FArchive& Reader, TArray<uint8>& IndexData, FSHAHash& InExpectedHash, FSHAHash& OutActualHash);
+	PAKFILE_API bool DecryptAndValidateIndex(FArchive& Reader, TArray<uint8>& IndexData, FSHAHash& InExpectedHash, FSHAHash& OutActualHash);
 
 	/* Manually add a file to a pak file */
-	void AddSpecialFile(const FPakEntry& Entry, const FString& Filename);
+	PAKFILE_API void AddSpecialFile(const FPakEntry& Entry, const FString& Filename);
 
 	/**
 	 * Search the given FDirectoryIndex for all files under the given Directory.  Helper for FindFilesAtPath, called separately on the DirectoryIndex or Pruned DirectoryIndex. Does not use
@@ -1755,24 +1755,24 @@ private:
 	}
 
 	/** Converts the path to a RelativePathFromMount and normalizes it to the expected format for Pak Directories.  Returns false if Path is not under the MountDir and hence can not be in this PakFile. */
-	bool NormalizeDirectoryQuery(const TCHAR* InPath, FString& OutRelativePathFromMount) const;
+	PAKFILE_API bool NormalizeDirectoryQuery(const TCHAR* InPath, FString& OutRelativePathFromMount) const;
 
 	/**
 	 * Looks up the given Normalized RelativePath in the Pruned DirectoryIndex and returns the directory if found.
 	 * Validates the result if IsPakValidatePruning. Does not use a critical section; caller is responsible for calling from within an FScopedPakDirectoryIndexAccess
 	 */
-	const FPakDirectory* FindPrunedDirectoryInternal(const FString& RelativePathFromMount) const;
+	PAKFILE_API const FPakDirectory* FindPrunedDirectoryInternal(const FString& RelativePathFromMount) const;
 
 #if ENABLE_PAKFILE_RUNTIME_PRUNING_VALIDATE
 	/* Logs an error if the two sets are not identical after removing all config-specified ignore paths */
-	void ValidateDirectorySearch(const TSet<FString>& FoundFullFiles, const TSet<FString>& PrunedFoundFiles, const TCHAR* InPath) const;
+	PAKFILE_API void ValidateDirectorySearch(const TSet<FString>& FoundFullFiles, const TSet<FString>& PrunedFoundFiles, const TCHAR* InPath) const;
 #endif
 };
 
 /**
  * Placeholder Class
  */
-class PAKFILE_API FPakNoEncryption
+class FPakNoEncryption
 {
 public:
 	enum 
@@ -1797,7 +1797,7 @@ public:
 typedef TFunction<FSharedPakReader()> TAcquirePakReaderFunction;
 
 template< typename EncryptionPolicy = FPakNoEncryption >
-class PAKFILE_API FPakReaderPolicy
+class FPakReaderPolicy
 {
 public:
 	/** Pak file that own this file data */
@@ -1867,7 +1867,7 @@ public:
  * File handle to read from pak file.
  */
 template< typename ReaderPolicy = FPakReaderPolicy<> >
-class PAKFILE_API FPakFileHandle : public IFileHandle
+class FPakFileHandle : public IFileHandle
 {
 	/** Current read position. */
 	int64 ReadPos;
diff --git a/Engine/Source/Runtime/PerfCounters/Public/PerfCountersModule.h b/Engine/Source/Runtime/PerfCounters/Public/PerfCountersModule.h
--- a/Engine/Source/Runtime/PerfCounters/Public/PerfCountersModule.h
+++ b/Engine/Source/Runtime/PerfCounters/Public/PerfCountersModule.h
@@ -44,7 +44,7 @@ DECLARE_DELEGATE_RetVal_TwoParams(bool, FPerfCounterExecCommandCallback, const F
 /**
  * A programming interface for setting/updating performance counters
  */
-class PERFCOUNTERS_API IPerfCounters
+class IPerfCounters
 {
 public:
 
@@ -66,21 +66,21 @@ public:
 	};
 
 	/** Named engine-wide histograms */
-	struct PERFCOUNTERS_API Histograms
+	struct Histograms
 	{
 		/** Frame time histogram for the duration of the match. */
-		static const FName FrameTime;
+		static PERFCOUNTERS_API const FName FrameTime;
 		/** Frame time histogram for shorter intervals. */
-		static const FName FrameTimePeriodic;
+		static PERFCOUNTERS_API const FName FrameTimePeriodic;
 		/** Frame time histogram (without sleep) for the duration of the match. */
-		static const FName FrameTimeWithoutSleep;
+		static PERFCOUNTERS_API const FName FrameTimeWithoutSleep;
 		/** ServerReplicateActors time histogram for the duration of the match. */
-		static const FName ServerReplicateActorsTime;
+		static PERFCOUNTERS_API const FName ServerReplicateActorsTime;
 		/** Sleep time histogram for the duration of the match. */
-		static const FName SleepTime;
+		static PERFCOUNTERS_API const FName SleepTime;
 
 		/** Zero load thread frame time histogram. */
-		static const FName ZeroLoadFrameTime;
+		static PERFCOUNTERS_API const FName ZeroLoadFrameTime;
 	};
 
 	/** Array used to store performance histograms. */
diff --git a/Engine/Source/Runtime/PhysicsCore/Public/BodySetupCore.h b/Engine/Source/Runtime/PhysicsCore/Public/BodySetupCore.h
--- a/Engine/Source/Runtime/PhysicsCore/Public/BodySetupCore.h
+++ b/Engine/Source/Runtime/PhysicsCore/Public/BodySetupCore.h
@@ -14,13 +14,13 @@ namespace physx
 	class PxTriangleMesh;
 }
 
-UCLASS(collapseCategories)
-class PHYSICSCORE_API UBodySetupCore : public UObject
+UCLASS(collapseCategories, MinimalAPI)
+class UBodySetupCore : public UObject
 {
 	GENERATED_UCLASS_BODY()
 
 public:
-	TEnumAsByte<enum ECollisionTraceFlag> GetCollisionTraceFlag() const;
+	PHYSICSCORE_API TEnumAsByte<enum ECollisionTraceFlag> GetCollisionTraceFlag() const;
 
 	/** Used in the PhysicsAsset case. Associates this Body with Bone in a skeletal mesh. */
 	UPROPERTY(Category=BodySetup,VisibleAnywhere)
@@ -39,4 +39,4 @@ public:
 	/** Collision Type for this body. This eventually changes response to collision to others **/
 	UPROPERTY(EditAnywhere,Category=Collision)
 	TEnumAsByte<enum EBodyCollisionResponse::Type> CollisionReponse;
-};
\ No newline at end of file
+};
diff --git a/Engine/Source/Runtime/PhysicsCore/Public/Chaos/ChaosConstraintSettings.h b/Engine/Source/Runtime/PhysicsCore/Public/Chaos/ChaosConstraintSettings.h
--- a/Engine/Source/Runtime/PhysicsCore/Public/Chaos/ChaosConstraintSettings.h
+++ b/Engine/Source/Runtime/PhysicsCore/Public/Chaos/ChaosConstraintSettings.h
@@ -12,26 +12,26 @@ namespace Chaos
 	*  that converting between existing (pre v4.26) simulation parameters and Chaos can produce similar 
 	*  simulation results. 
 	*/
-	class PHYSICSCORE_API ConstraintSettings
+	class ConstraintSettings
 	{
 		ConstraintSettings(){};
 
 	public:
-		static FReal JointStiffness();
-		static FReal LinearDriveStiffnessScale();
-		static FReal LinearDriveDampingScale();
-		static FReal AngularDriveStiffnessScale();
-		static FReal AngularDriveDampingScale();
+		static PHYSICSCORE_API FReal JointStiffness();
+		static PHYSICSCORE_API FReal LinearDriveStiffnessScale();
+		static PHYSICSCORE_API FReal LinearDriveDampingScale();
+		static PHYSICSCORE_API FReal AngularDriveStiffnessScale();
+		static PHYSICSCORE_API FReal AngularDriveDampingScale();
 
-		static int SoftLinearForceMode();
-		static FReal SoftLinearStiffnessScale();
-		static FReal SoftLinearDampingScale();
+		static PHYSICSCORE_API int SoftLinearForceMode();
+		static PHYSICSCORE_API FReal SoftLinearStiffnessScale();
+		static PHYSICSCORE_API FReal SoftLinearDampingScale();
 
-		static int SoftAngularForceMode();
-		static FReal SoftAngularStiffnessScale();
-		static FReal SoftAngularDampingScale();
+		static PHYSICSCORE_API int SoftAngularForceMode();
+		static PHYSICSCORE_API FReal SoftAngularStiffnessScale();
+		static PHYSICSCORE_API FReal SoftAngularDampingScale();
 
-		static FReal LinearBreakScale();
-		static FReal AngularBreakScale();
+		static PHYSICSCORE_API FReal LinearBreakScale();
+		static PHYSICSCORE_API FReal AngularBreakScale();
 	};
-}
\ No newline at end of file
+}
diff --git a/Engine/Source/Runtime/PhysicsCore/Public/Chaos/ChaosEngineInterface.h b/Engine/Source/Runtime/PhysicsCore/Public/Chaos/ChaosEngineInterface.h
--- a/Engine/Source/Runtime/PhysicsCore/Public/Chaos/ChaosEngineInterface.h
+++ b/Engine/Source/Runtime/PhysicsCore/Public/Chaos/ChaosEngineInterface.h
@@ -230,19 +230,19 @@ namespace Chaos
 struct FCollisionShape;
 class UPhysicalMaterial;
 
-class PHYSICSCORE_API FPhysicsAggregateReference_Chaos
+class FPhysicsAggregateReference_Chaos
 {
 public:
 	bool IsValid() const { return false; }
 };
 
-class PHYSICSCORE_API FPhysicsConstraintReference_Chaos
+class FPhysicsConstraintReference_Chaos
 {
 public:
 	FPhysicsConstraintReference_Chaos() { Reset(); }
 	void Reset() { Constraint = nullptr; }
 
-	bool IsValid() const;
+	PHYSICSCORE_API bool IsValid() const;
 
 
 
@@ -252,7 +252,7 @@ public:
 	Chaos::FConstraintBase* Constraint;
 };
 
-class PHYSICSCORE_API FPhysicsShapeReference_Chaos
+class FPhysicsShapeReference_Chaos
 {
 public:
 	
@@ -264,7 +264,7 @@ public:
 	bool IsValid() const { return (Shape != nullptr); }
 	bool Equals(const FPhysicsShapeReference_Chaos& Other) const { return Shape == Other.Shape; }
     bool operator==(const FPhysicsShapeReference_Chaos& Other) const { return Equals(Other); }
-	const Chaos::FImplicitObject& GetGeometry() const;
+	PHYSICSCORE_API const Chaos::FImplicitObject& GetGeometry() const;
 
 	Chaos::FPerShapeData* Shape;
     FPhysicsActorHandle ActorRef;
@@ -275,15 +275,15 @@ public:
 	}
 };
 
-class PHYSICSCORE_API FPhysicsShapeAdapter_Chaos
+class FPhysicsShapeAdapter_Chaos
 {
 public:
-	FPhysicsShapeAdapter_Chaos(const FQuat& Rot, const FCollisionShape& CollisionShape);
-	~FPhysicsShapeAdapter_Chaos();
+	PHYSICSCORE_API FPhysicsShapeAdapter_Chaos(const FQuat& Rot, const FCollisionShape& CollisionShape);
+	PHYSICSCORE_API ~FPhysicsShapeAdapter_Chaos();
 
-	const FPhysicsGeometry& GetGeometry() const;
-	FTransform GetGeomPose(const FVector& Pos) const;
-	const FQuat& GetGeomOrientation() const;
+	PHYSICSCORE_API const FPhysicsGeometry& GetGeometry() const;
+	PHYSICSCORE_API FTransform GetGeomPose(const FVector& Pos) const;
+	PHYSICSCORE_API const FQuat& GetGeomOrientation() const;
 
 private:
 	TUniquePtr<FPhysicsGeometry> Geometry;
@@ -293,7 +293,7 @@ private:
 /**
  Wrapper around geometry. This is really just needed to make the physx chaos abstraction easier
  */
-struct PHYSICSCORE_API FPhysicsGeometryCollection_Chaos
+struct FPhysicsGeometryCollection_Chaos
 {
 	// Delete default constructor, want only construction by interface (private constructor below)
 	FPhysicsGeometryCollection_Chaos() = delete;
@@ -301,22 +301,22 @@ struct PHYSICSCORE_API FPhysicsGeometryCollection_Chaos
 	// to be able to delete physx::PxGeometryHolder which is incomplete here
 	FPhysicsGeometryCollection_Chaos(const FPhysicsGeometryCollection_Chaos& Copy) = delete;
 	FPhysicsGeometryCollection_Chaos& operator=(const FPhysicsGeometryCollection_Chaos& Copy) = delete;
-	FPhysicsGeometryCollection_Chaos(FPhysicsGeometryCollection_Chaos&& Steal);
+	PHYSICSCORE_API FPhysicsGeometryCollection_Chaos(FPhysicsGeometryCollection_Chaos&& Steal);
 	FPhysicsGeometryCollection_Chaos& operator=(FPhysicsGeometryCollection_Chaos&& Steal) = delete;
-	~FPhysicsGeometryCollection_Chaos();
+	PHYSICSCORE_API ~FPhysicsGeometryCollection_Chaos();
 
-	ECollisionShapeType GetType() const;
-	const Chaos::FImplicitObject& GetGeometry() const;
-	const Chaos::TBox<Chaos::FReal, 3>& GetBoxGeometry() const;
-	const Chaos::TSphere<Chaos::FReal, 3>&  GetSphereGeometry() const;
-	const Chaos::FCapsule&  GetCapsuleGeometry() const;
-	const Chaos::FConvex& GetConvexGeometry() const;
-	const Chaos::FTriangleMeshImplicitObject& GetTriMeshGeometry() const;
+	PHYSICSCORE_API ECollisionShapeType GetType() const;
+	PHYSICSCORE_API const Chaos::FImplicitObject& GetGeometry() const;
+	PHYSICSCORE_API const Chaos::TBox<Chaos::FReal, 3>& GetBoxGeometry() const;
+	PHYSICSCORE_API const Chaos::TSphere<Chaos::FReal, 3>&  GetSphereGeometry() const;
+	PHYSICSCORE_API const Chaos::FCapsule&  GetCapsuleGeometry() const;
+	PHYSICSCORE_API const Chaos::FConvex& GetConvexGeometry() const;
+	PHYSICSCORE_API const Chaos::FTriangleMeshImplicitObject& GetTriMeshGeometry() const;
 
 private:
 	friend class FChaosEngineInterface;
-	explicit FPhysicsGeometryCollection_Chaos(const FPhysicsShapeReference_Chaos& InShape);
-	explicit FPhysicsGeometryCollection_Chaos(const FPhysicsGeometry& InGeom);
+	PHYSICSCORE_API explicit FPhysicsGeometryCollection_Chaos(const FPhysicsShapeReference_Chaos& InShape);
+	PHYSICSCORE_API explicit FPhysicsGeometryCollection_Chaos(const FPhysicsGeometry& InGeom);
 
 	const Chaos::FImplicitObject& Geom;
 };
@@ -336,218 +336,218 @@ namespace physx
 struct FContactModifyCallback;
 class ULineBatchComponent;
 
-class PHYSICSCORE_API FChaosEngineInterface
+class FChaosEngineInterface
 {
 public:
 	virtual ~FChaosEngineInterface() = default;
 
-	static void CreateActor(const FActorCreationParams& InParams,FPhysicsActorHandle& Handle);
-	static void ReleaseActor(FPhysicsActorHandle& InActorReference,FChaosScene* InScene = nullptr,bool bNeverDeferRelease=false);
+	static PHYSICSCORE_API void CreateActor(const FActorCreationParams& InParams,FPhysicsActorHandle& Handle);
+	static PHYSICSCORE_API void ReleaseActor(FPhysicsActorHandle& InActorReference,FChaosScene* InScene = nullptr,bool bNeverDeferRelease=false);
 
 	static bool IsValid(const FPhysicsActorHandle& Handle) { return Handle != nullptr; }
-	static void AddActorToSolver(FPhysicsActorHandle& Handle,Chaos::FPhysicsSolver* Solver);
-	static void RemoveActorFromSolver(FPhysicsActorHandle& Handle,Chaos::FPhysicsSolver* Solver);
+	static PHYSICSCORE_API void AddActorToSolver(FPhysicsActorHandle& Handle,Chaos::FPhysicsSolver* Solver);
+	static PHYSICSCORE_API void RemoveActorFromSolver(FPhysicsActorHandle& Handle,Chaos::FPhysicsSolver* Solver);
 
-	static void SetGlobalPose_AssumesLocked(const FPhysicsActorHandle& InActorReference,const FTransform& InNewPose,bool bAutoWake = true);
-	static void SetKinematicTarget_AssumesLocked(const FPhysicsActorHandle& InActorReference,const FTransform& InNewTarget);
+	static PHYSICSCORE_API void SetGlobalPose_AssumesLocked(const FPhysicsActorHandle& InActorReference,const FTransform& InNewPose,bool bAutoWake = true);
+	static PHYSICSCORE_API void SetKinematicTarget_AssumesLocked(const FPhysicsActorHandle& InActorReference,const FTransform& InNewTarget);
 
-	static FChaosScene* GetCurrentScene(const FPhysicsActorHandle& InHandle);
+	static PHYSICSCORE_API FChaosScene* GetCurrentScene(const FPhysicsActorHandle& InHandle);
 
-	static FPhysicsMaterialHandle CreateMaterial(const UPhysicalMaterial* InMaterial);
-	static void UpdateMaterial(FPhysicsMaterialHandle& InHandle, UPhysicalMaterial* InMaterial);
-	static void ReleaseMaterial(FPhysicsMaterialHandle& InHandle);
-	static void SetUserData(FPhysicsMaterialHandle& InHandle,void* InUserData);
+	static PHYSICSCORE_API FPhysicsMaterialHandle CreateMaterial(const UPhysicalMaterial* InMaterial);
+	static PHYSICSCORE_API void UpdateMaterial(FPhysicsMaterialHandle& InHandle, UPhysicalMaterial* InMaterial);
+	static PHYSICSCORE_API void ReleaseMaterial(FPhysicsMaterialHandle& InHandle);
+	static PHYSICSCORE_API void SetUserData(FPhysicsMaterialHandle& InHandle,void* InUserData);
 
-	static void ReleaseMaterialMask(FPhysicsMaterialMaskHandle& InHandle);
+	static PHYSICSCORE_API void ReleaseMaterialMask(FPhysicsMaterialMaskHandle& InHandle);
 
-	static FPhysicsAggregateReference_Chaos CreateAggregate(int32 MaxBodies);
-	static void ReleaseAggregate(FPhysicsAggregateReference_Chaos& InAggregate);
-	static int32 GetNumActorsInAggregate(const FPhysicsAggregateReference_Chaos& InAggregate);
-	static void AddActorToAggregate_AssumesLocked(const FPhysicsAggregateReference_Chaos& InAggregate,const FPhysicsActorHandle& InActor);
+	static PHYSICSCORE_API FPhysicsAggregateReference_Chaos CreateAggregate(int32 MaxBodies);
+	static PHYSICSCORE_API void ReleaseAggregate(FPhysicsAggregateReference_Chaos& InAggregate);
+	static PHYSICSCORE_API int32 GetNumActorsInAggregate(const FPhysicsAggregateReference_Chaos& InAggregate);
+	static PHYSICSCORE_API void AddActorToAggregate_AssumesLocked(const FPhysicsAggregateReference_Chaos& InAggregate,const FPhysicsActorHandle& InActor);
 
-	static int32 GetNumShapes(const FPhysicsActorHandle& InHandle);
+	static PHYSICSCORE_API int32 GetNumShapes(const FPhysicsActorHandle& InHandle);
 
-	static void ReleaseShape(const FPhysicsShapeHandle& InShape);
+	static PHYSICSCORE_API void ReleaseShape(const FPhysicsShapeHandle& InShape);
 
-	static void AttachShape(const FPhysicsActorHandle& InActor,const FPhysicsShapeHandle& InNewShape);
-	static void DetachShape(const FPhysicsActorHandle& InActor,FPhysicsShapeHandle& InShape,bool bWakeTouching = true);
+	static PHYSICSCORE_API void AttachShape(const FPhysicsActorHandle& InActor,const FPhysicsShapeHandle& InNewShape);
+	static PHYSICSCORE_API void DetachShape(const FPhysicsActorHandle& InActor,FPhysicsShapeHandle& InShape,bool bWakeTouching = true);
 
-	static void SetSmoothEdgeCollisionsEnabled_AssumesLocked(const FPhysicsActorHandle& InActor, const bool bSmoothEdgeCollisionsEnabled);
+	static PHYSICSCORE_API void SetSmoothEdgeCollisionsEnabled_AssumesLocked(const FPhysicsActorHandle& InActor, const bool bSmoothEdgeCollisionsEnabled);
 
-	static void AddDisabledCollisionsFor_AssumesLocked(const TMap<FPhysicsActorHandle, TArray< FPhysicsActorHandle > >& InMap);
-	static void RemoveDisabledCollisionsFor_AssumesLocked(TArray< FPhysicsActorHandle > & InPhysicsActors);
+	static PHYSICSCORE_API void AddDisabledCollisionsFor_AssumesLocked(const TMap<FPhysicsActorHandle, TArray< FPhysicsActorHandle > >& InMap);
+	static PHYSICSCORE_API void RemoveDisabledCollisionsFor_AssumesLocked(TArray< FPhysicsActorHandle > & InPhysicsActors);
 
-	static void SetDisabled(const FPhysicsActorHandle& InPhysicsActor, bool bSetDisabled);
-	static bool IsDisabled(const FPhysicsActorHandle& InPhysicsActor);
+	static PHYSICSCORE_API void SetDisabled(const FPhysicsActorHandle& InPhysicsActor, bool bSetDisabled);
+	static PHYSICSCORE_API bool IsDisabled(const FPhysicsActorHandle& InPhysicsActor);
 
-	static void SetActorUserData_AssumesLocked(FPhysicsActorHandle& InActorReference,FPhysicsUserData* InUserData);
+	static PHYSICSCORE_API void SetActorUserData_AssumesLocked(FPhysicsActorHandle& InActorReference,FPhysicsUserData* InUserData);
 
-	static bool IsRigidBody(const FPhysicsActorHandle& InActorReference);
-	static bool IsDynamic(const FPhysicsActorHandle& InActorReference);
-	static bool IsStatic(const FPhysicsActorHandle& InActorReference);
-	static bool IsKinematic(const FPhysicsActorHandle& InActorReference);
-	static bool IsKinematic_AssumesLocked(const FPhysicsActorHandle& InActorReference);
-	static bool IsSleeping(const FPhysicsActorHandle& InActorReference);
-	static bool IsCcdEnabled(const FPhysicsActorHandle& InActorReference);
+	static PHYSICSCORE_API bool IsRigidBody(const FPhysicsActorHandle& InActorReference);
+	static PHYSICSCORE_API bool IsDynamic(const FPhysicsActorHandle& InActorReference);
+	static PHYSICSCORE_API bool IsStatic(const FPhysicsActorHandle& InActorReference);
+	static PHYSICSCORE_API bool IsKinematic(const FPhysicsActorHandle& InActorReference);
+	static PHYSICSCORE_API bool IsKinematic_AssumesLocked(const FPhysicsActorHandle& InActorReference);
+	static PHYSICSCORE_API bool IsSleeping(const FPhysicsActorHandle& InActorReference);
+	static PHYSICSCORE_API bool IsCcdEnabled(const FPhysicsActorHandle& InActorReference);
 	// @todo(mlentine): We don't have a notion of sync vs async and are a bit of both. Does this work?
 	static bool HasSyncSceneData(const FPhysicsActorHandle& InHandle) { return true; }
 	static bool HasAsyncSceneData(const FPhysicsActorHandle& InHandle) { return false; }
-	static bool IsInScene(const FPhysicsActorHandle& InActorReference);
+	static PHYSICSCORE_API bool IsInScene(const FPhysicsActorHandle& InActorReference);
 
-	static int32 GetAllShapes_AssumedLocked(const FPhysicsActorHandle& InActorHandle,TArray<FPhysicsShapeReference_Chaos,FDefaultAllocator>& OutShapes);
-	static int32 GetAllShapes_AssumedLocked(const FPhysicsActorHandle& InActorHandle,PhysicsInterfaceTypes::FInlineShapeArray& OutShapes);
+	static PHYSICSCORE_API int32 GetAllShapes_AssumedLocked(const FPhysicsActorHandle& InActorHandle,TArray<FPhysicsShapeReference_Chaos,FDefaultAllocator>& OutShapes);
+	static PHYSICSCORE_API int32 GetAllShapes_AssumedLocked(const FPhysicsActorHandle& InActorHandle,PhysicsInterfaceTypes::FInlineShapeArray& OutShapes);
 
-	static bool CanSimulate_AssumesLocked(const FPhysicsActorHandle& InActorReference);
-	static float GetMass_AssumesLocked(const FPhysicsActorHandle& InActorReference);
+	static PHYSICSCORE_API bool CanSimulate_AssumesLocked(const FPhysicsActorHandle& InActorReference);
+	static PHYSICSCORE_API float GetMass_AssumesLocked(const FPhysicsActorHandle& InActorReference);
 
-	static void SetSendsSleepNotifies_AssumesLocked(const FPhysicsActorHandle& InActorReference,bool bSendSleepNotifies);
-	static void PutToSleep_AssumesLocked(const FPhysicsActorHandle& InActorReference);
-	static void WakeUp_AssumesLocked(const FPhysicsActorHandle& InActorReference);
+	static PHYSICSCORE_API void SetSendsSleepNotifies_AssumesLocked(const FPhysicsActorHandle& InActorReference,bool bSendSleepNotifies);
+	static PHYSICSCORE_API void PutToSleep_AssumesLocked(const FPhysicsActorHandle& InActorReference);
+	static PHYSICSCORE_API void WakeUp_AssumesLocked(const FPhysicsActorHandle& InActorReference);
 
-	static void SetIsKinematic_AssumesLocked(const FPhysicsActorHandle& InActorReference,bool bIsKinematic);
-	static void SetCcdEnabled_AssumesLocked(const FPhysicsActorHandle& InActorReference,bool bIsCcdEnabled);
-	static void SetIgnoreAnalyticCollisions_AssumesLocked(const FPhysicsActorHandle& InActorReference,bool bIsCcdEnabled);
+	static PHYSICSCORE_API void SetIsKinematic_AssumesLocked(const FPhysicsActorHandle& InActorReference,bool bIsKinematic);
+	static PHYSICSCORE_API void SetCcdEnabled_AssumesLocked(const FPhysicsActorHandle& InActorReference,bool bIsCcdEnabled);
+	static PHYSICSCORE_API void SetIgnoreAnalyticCollisions_AssumesLocked(const FPhysicsActorHandle& InActorReference,bool bIsCcdEnabled);
 
-	static FTransform GetGlobalPose_AssumesLocked(const FPhysicsActorHandle& InActorReference);
+	static PHYSICSCORE_API FTransform GetGlobalPose_AssumesLocked(const FPhysicsActorHandle& InActorReference);
 
-	static FTransform GetTransform_AssumesLocked(const FPhysicsActorHandle& InRef,bool bForceGlobalPose = false);
+	static PHYSICSCORE_API FTransform GetTransform_AssumesLocked(const FPhysicsActorHandle& InRef,bool bForceGlobalPose = false);
 
-	static bool HasKinematicTarget_AssumesLocked(const FPhysicsActorHandle& InActorReference);
-	static FTransform GetKinematicTarget_AssumesLocked(const FPhysicsActorHandle& InActorReference);
+	static PHYSICSCORE_API bool HasKinematicTarget_AssumesLocked(const FPhysicsActorHandle& InActorReference);
+	static PHYSICSCORE_API FTransform GetKinematicTarget_AssumesLocked(const FPhysicsActorHandle& InActorReference);
 
-	static FVector GetLinearVelocity_AssumesLocked(const FPhysicsActorHandle& InActorReference);
-	static void SetLinearVelocity_AssumesLocked(const FPhysicsActorHandle& InActorReference,const FVector& InNewVelocity,bool bAutoWake = true);
+	static PHYSICSCORE_API FVector GetLinearVelocity_AssumesLocked(const FPhysicsActorHandle& InActorReference);
+	static PHYSICSCORE_API void SetLinearVelocity_AssumesLocked(const FPhysicsActorHandle& InActorReference,const FVector& InNewVelocity,bool bAutoWake = true);
 
-	static FVector GetAngularVelocity_AssumesLocked(const FPhysicsActorHandle& InActorReference);
-	static void SetAngularVelocity_AssumesLocked(const FPhysicsActorHandle& InActorReference,const FVector& InNewVelocity,bool bAutoWake = true);
-	static float GetMaxAngularVelocity_AssumesLocked(const FPhysicsActorHandle& InActorReference);
-	static float GetMaxLinearVelocity_AssumesLocked(const FPhysicsActorHandle& InActorReference);
-	static void SetMaxAngularVelocity_AssumesLocked(const FPhysicsActorHandle& InActorReference,float InMaxAngularVelocityRadians);
-	static void SetMaxLinearVelocity_AssumesLocked(const FPhysicsActorHandle& InActorReference, float InMaxLinearVelocity);
+	static PHYSICSCORE_API FVector GetAngularVelocity_AssumesLocked(const FPhysicsActorHandle& InActorReference);
+	static PHYSICSCORE_API void SetAngularVelocity_AssumesLocked(const FPhysicsActorHandle& InActorReference,const FVector& InNewVelocity,bool bAutoWake = true);
+	static PHYSICSCORE_API float GetMaxAngularVelocity_AssumesLocked(const FPhysicsActorHandle& InActorReference);
+	static PHYSICSCORE_API float GetMaxLinearVelocity_AssumesLocked(const FPhysicsActorHandle& InActorReference);
+	static PHYSICSCORE_API void SetMaxAngularVelocity_AssumesLocked(const FPhysicsActorHandle& InActorReference,float InMaxAngularVelocityRadians);
+	static PHYSICSCORE_API void SetMaxLinearVelocity_AssumesLocked(const FPhysicsActorHandle& InActorReference, float InMaxLinearVelocity);
 
-	static float GetMaxDepenetrationVelocity_AssumesLocked(const FPhysicsActorHandle& InActorReference);
-	static void SetMaxDepenetrationVelocity_AssumesLocked(const FPhysicsActorHandle& InActorReference,float InMaxDepenetrationVelocity);
+	static PHYSICSCORE_API float GetMaxDepenetrationVelocity_AssumesLocked(const FPhysicsActorHandle& InActorReference);
+	static PHYSICSCORE_API void SetMaxDepenetrationVelocity_AssumesLocked(const FPhysicsActorHandle& InActorReference,float InMaxDepenetrationVelocity);
 
-	static FVector GetWorldVelocityAtPoint_AssumesLocked(const FPhysicsActorHandle& InActorReference,const FVector& InPoint);
-	static FVector GetWorldVelocityAtPoint_AssumesLocked(const Chaos::FRigidBodyHandle_Internal* InActorReference, const FVector& InPoint);
+	static PHYSICSCORE_API FVector GetWorldVelocityAtPoint_AssumesLocked(const FPhysicsActorHandle& InActorReference,const FVector& InPoint);
+	static PHYSICSCORE_API FVector GetWorldVelocityAtPoint_AssumesLocked(const Chaos::FRigidBodyHandle_Internal* InActorReference, const FVector& InPoint);
 
-	static FTransform GetComTransform_AssumesLocked(const FPhysicsActorHandle& InActorReference);
-	static FTransform GetComTransformLocal_AssumesLocked(const FPhysicsActorHandle& InActorReference);
+	static PHYSICSCORE_API FTransform GetComTransform_AssumesLocked(const FPhysicsActorHandle& InActorReference);
+	static PHYSICSCORE_API FTransform GetComTransformLocal_AssumesLocked(const FPhysicsActorHandle& InActorReference);
 
-	static FVector GetLocalInertiaTensor_AssumesLocked(const FPhysicsActorHandle& InActorReference);
-	static FBox GetBounds_AssumesLocked(const FPhysicsActorHandle& InActorReference);
-	static FBox GetBounds_AssumesLocked(const FPhysicsActorHandle& InActorReference, const FTransform& InTransform);
+	static PHYSICSCORE_API FVector GetLocalInertiaTensor_AssumesLocked(const FPhysicsActorHandle& InActorReference);
+	static PHYSICSCORE_API FBox GetBounds_AssumesLocked(const FPhysicsActorHandle& InActorReference);
+	static PHYSICSCORE_API FBox GetBounds_AssumesLocked(const FPhysicsActorHandle& InActorReference, const FTransform& InTransform);
 
-	static void SetLinearDamping_AssumesLocked(const FPhysicsActorHandle& InActorReference,float InDamping);
-	static void SetAngularDamping_AssumesLocked(const FPhysicsActorHandle& InActorReference,float InDamping);
+	static PHYSICSCORE_API void SetLinearDamping_AssumesLocked(const FPhysicsActorHandle& InActorReference,float InDamping);
+	static PHYSICSCORE_API void SetAngularDamping_AssumesLocked(const FPhysicsActorHandle& InActorReference,float InDamping);
 
-	static void AddForce_AssumesLocked(const FPhysicsActorHandle& InActorReference, const FVector& Force, bool bAllowSubstepping, bool bAccelChange, bool bIsInternal = false);
-	static void AddForceAtPosition_AssumesLocked(const FPhysicsActorHandle& InActorReference, const FVector& Force, const FVector& Position, bool bAllowSubstepping, bool bIsLocalForce = false, bool bIsInternal = false);
-	static void AddRadialForce_AssumesLocked(const FPhysicsActorHandle& InActorReference, const FVector& Origin, const float Radius, const float Strength, const uint8 Falloff, bool bAccelChange, bool bAllowSubstepping, bool bIsInternal = false);
-	static void AddTorque_AssumesLocked(const FPhysicsActorHandle& InActorReference, const FVector& Torque, bool bAllowSubstepping, bool bAccelChange, bool bIsInternal = false);
+	static PHYSICSCORE_API void AddForce_AssumesLocked(const FPhysicsActorHandle& InActorReference, const FVector& Force, bool bAllowSubstepping, bool bAccelChange, bool bIsInternal = false);
+	static PHYSICSCORE_API void AddForceAtPosition_AssumesLocked(const FPhysicsActorHandle& InActorReference, const FVector& Force, const FVector& Position, bool bAllowSubstepping, bool bIsLocalForce = false, bool bIsInternal = false);
+	static PHYSICSCORE_API void AddRadialForce_AssumesLocked(const FPhysicsActorHandle& InActorReference, const FVector& Origin, const float Radius, const float Strength, const uint8 Falloff, bool bAccelChange, bool bAllowSubstepping, bool bIsInternal = false);
+	static PHYSICSCORE_API void AddTorque_AssumesLocked(const FPhysicsActorHandle& InActorReference, const FVector& Torque, bool bAllowSubstepping, bool bAccelChange, bool bIsInternal = false);
 
-	static void AddImpulse_AssumesLocked(const FPhysicsActorHandle& InActorReference,const FVector& InForce, bool bIsInternal = false);
-	static void AddAngularImpulseInRadians_AssumesLocked(const FPhysicsActorHandle& InActorReference,const FVector& InTorque, bool bIsInternal = false);
-	static void AddVelocity_AssumesLocked(const FPhysicsActorHandle& InActorReference,const FVector& InForce, bool bIsInternal = false);
-	static void AddAngularVelocityInRadians_AssumesLocked(const FPhysicsActorHandle& InActorReference,const FVector& InTorque, bool bIsInternal = false);
-	static void AddImpulseAtLocation_AssumesLocked(const FPhysicsActorHandle& InActorReference,const FVector& InImpulse,const FVector& InLocation, bool bIsInternal = false);
-	static void AddVelocityChangeImpulseAtLocation_AssumesLocked(const FPhysicsActorHandle& InActorReference, const FVector& InVelocityDelta, const FVector& InLocation, bool bIsInternal = false);
-	static void AddRadialImpulse_AssumesLocked(const FPhysicsActorHandle& InActorReference,const FVector& InOrigin,float InRadius,float InStrength,ERadialImpulseFalloff InFalloff,bool bInVelChange, bool bIsInternal = false);
+	static PHYSICSCORE_API void AddImpulse_AssumesLocked(const FPhysicsActorHandle& InActorReference,const FVector& InForce, bool bIsInternal = false);
+	static PHYSICSCORE_API void AddAngularImpulseInRadians_AssumesLocked(const FPhysicsActorHandle& InActorReference,const FVector& InTorque, bool bIsInternal = false);
+	static PHYSICSCORE_API void AddVelocity_AssumesLocked(const FPhysicsActorHandle& InActorReference,const FVector& InForce, bool bIsInternal = false);
+	static PHYSICSCORE_API void AddAngularVelocityInRadians_AssumesLocked(const FPhysicsActorHandle& InActorReference,const FVector& InTorque, bool bIsInternal = false);
+	static PHYSICSCORE_API void AddImpulseAtLocation_AssumesLocked(const FPhysicsActorHandle& InActorReference,const FVector& InImpulse,const FVector& InLocation, bool bIsInternal = false);
+	static PHYSICSCORE_API void AddVelocityChangeImpulseAtLocation_AssumesLocked(const FPhysicsActorHandle& InActorReference, const FVector& InVelocityDelta, const FVector& InLocation, bool bIsInternal = false);
+	static PHYSICSCORE_API void AddRadialImpulse_AssumesLocked(const FPhysicsActorHandle& InActorReference,const FVector& InOrigin,float InRadius,float InStrength,ERadialImpulseFalloff InFalloff,bool bInVelChange, bool bIsInternal = false);
 
-	static bool IsGravityEnabled_AssumesLocked(const FPhysicsActorHandle& InActorReference);
-	static void SetGravityEnabled_AssumesLocked(const FPhysicsActorHandle& InActorReference,bool bEnabled);
+	static PHYSICSCORE_API bool IsGravityEnabled_AssumesLocked(const FPhysicsActorHandle& InActorReference);
+	static PHYSICSCORE_API void SetGravityEnabled_AssumesLocked(const FPhysicsActorHandle& InActorReference,bool bEnabled);
 
-	static bool GetUpdateKinematicFromSimulation_AssumesLocked(const FPhysicsActorHandle& InActorReference);
-	static void SetUpdateKinematicFromSimulation_AssumesLocked(const FPhysicsActorHandle& InActorReference, bool bUpdateKinematicFromSimulation);
+	static PHYSICSCORE_API bool GetUpdateKinematicFromSimulation_AssumesLocked(const FPhysicsActorHandle& InActorReference);
+	static PHYSICSCORE_API void SetUpdateKinematicFromSimulation_AssumesLocked(const FPhysicsActorHandle& InActorReference, bool bUpdateKinematicFromSimulation);
 
-	static void SetOneWayInteraction_AssumesLocked(const FPhysicsActorHandle& InHandle, bool InOneWayInteraction);
+	static PHYSICSCORE_API void SetOneWayInteraction_AssumesLocked(const FPhysicsActorHandle& InHandle, bool InOneWayInteraction);
 
-	static float GetSleepEnergyThreshold_AssumesLocked(const FPhysicsActorHandle& InActorReference);
-	static void SetSleepEnergyThreshold_AssumesLocked(const FPhysicsActorHandle& InActorReference,float InEnergyThreshold);
+	static PHYSICSCORE_API float GetSleepEnergyThreshold_AssumesLocked(const FPhysicsActorHandle& InActorReference);
+	static PHYSICSCORE_API void SetSleepEnergyThreshold_AssumesLocked(const FPhysicsActorHandle& InActorReference,float InEnergyThreshold);
 
-	static void SetMass_AssumesLocked(FPhysicsActorHandle& InHandle,float InMass);
-	static void SetMassSpaceInertiaTensor_AssumesLocked(FPhysicsActorHandle& InHandle,const FVector& InTensor);
-	static void SetComLocalPose_AssumesLocked(const FPhysicsActorHandle& InHandle,const FTransform& InComLocalPose);
+	static PHYSICSCORE_API void SetMass_AssumesLocked(FPhysicsActorHandle& InHandle,float InMass);
+	static PHYSICSCORE_API void SetMassSpaceInertiaTensor_AssumesLocked(FPhysicsActorHandle& InHandle,const FVector& InTensor);
+	static PHYSICSCORE_API void SetComLocalPose_AssumesLocked(const FPhysicsActorHandle& InHandle,const FTransform& InComLocalPose);
 
-	static bool IsInertiaConditioningEnabled_AssumesLocked(const FPhysicsActorHandle& InActorReference);
-	static void SetInertiaConditioningEnabled_AssumesLocked(const FPhysicsActorHandle& InActorReference, bool bEnabled);
+	static PHYSICSCORE_API bool IsInertiaConditioningEnabled_AssumesLocked(const FPhysicsActorHandle& InActorReference);
+	static PHYSICSCORE_API void SetInertiaConditioningEnabled_AssumesLocked(const FPhysicsActorHandle& InActorReference, bool bEnabled);
 
-	static float GetStabilizationEnergyThreshold_AssumesLocked(const FPhysicsActorHandle& InHandle);
-	static void SetStabilizationEnergyThreshold_AssumesLocked(const FPhysicsActorHandle& InHandle,float InThreshold);
-	static uint32 GetSolverPositionIterationCount_AssumesLocked(const FPhysicsActorHandle& InHandle);
-	static void SetSolverPositionIterationCount_AssumesLocked(const FPhysicsActorHandle& InHandle,uint32 InSolverIterationCount);
-	static uint32 GetSolverVelocityIterationCount_AssumesLocked(const FPhysicsActorHandle& InHandle);
-	static void SetSolverVelocityIterationCount_AssumesLocked(const FPhysicsActorHandle& InHandle,uint32 InSolverIterationCount);
-	static float GetWakeCounter_AssumesLocked(const FPhysicsActorHandle& InHandle);
-	static void SetWakeCounter_AssumesLocked(const FPhysicsActorHandle& InHandle,float InWakeCounter);
+	static PHYSICSCORE_API float GetStabilizationEnergyThreshold_AssumesLocked(const FPhysicsActorHandle& InHandle);
+	static PHYSICSCORE_API void SetStabilizationEnergyThreshold_AssumesLocked(const FPhysicsActorHandle& InHandle,float InThreshold);
+	static PHYSICSCORE_API uint32 GetSolverPositionIterationCount_AssumesLocked(const FPhysicsActorHandle& InHandle);
+	static PHYSICSCORE_API void SetSolverPositionIterationCount_AssumesLocked(const FPhysicsActorHandle& InHandle,uint32 InSolverIterationCount);
+	static PHYSICSCORE_API uint32 GetSolverVelocityIterationCount_AssumesLocked(const FPhysicsActorHandle& InHandle);
+	static PHYSICSCORE_API void SetSolverVelocityIterationCount_AssumesLocked(const FPhysicsActorHandle& InHandle,uint32 InSolverIterationCount);
+	static PHYSICSCORE_API float GetWakeCounter_AssumesLocked(const FPhysicsActorHandle& InHandle);
+	static PHYSICSCORE_API void SetWakeCounter_AssumesLocked(const FPhysicsActorHandle& InHandle,float InWakeCounter);
 
-	static void SetInitialized_AssumesLocked(const FPhysicsActorHandle& InHandle,bool InInitialized);
+	static PHYSICSCORE_API void SetInitialized_AssumesLocked(const FPhysicsActorHandle& InHandle,bool InInitialized);
 
-	static SIZE_T GetResourceSizeEx(const FPhysicsActorHandle& InActorRef);
+	static PHYSICSCORE_API SIZE_T GetResourceSizeEx(const FPhysicsActorHandle& InActorRef);
 
-	static FPhysicsConstraintHandle CreateConstraint(Chaos::FPhysicsObject* Body1, Chaos::FPhysicsObject* Body2, const FTransform& InLocalFrame1, const FTransform& InLocalFrame2);
-	static FPhysicsConstraintHandle CreateConstraint(const FPhysicsActorHandle& InActorRef1,const FPhysicsActorHandle& InActorRef2,const FTransform& InLocalFrame1,const FTransform& InLocalFrame2);
+	static PHYSICSCORE_API FPhysicsConstraintHandle CreateConstraint(Chaos::FPhysicsObject* Body1, Chaos::FPhysicsObject* Body2, const FTransform& InLocalFrame1, const FTransform& InLocalFrame2);
+	static PHYSICSCORE_API FPhysicsConstraintHandle CreateConstraint(const FPhysicsActorHandle& InActorRef1,const FPhysicsActorHandle& InActorRef2,const FTransform& InLocalFrame1,const FTransform& InLocalFrame2);
 
-	static FPhysicsConstraintHandle CreateSuspension(const FPhysicsActorHandle& InActorRef, const FVector& InLocalFrame);
-	static void SetConstraintUserData(const FPhysicsConstraintHandle& InConstraintRef,void* InUserData);
-	static void ReleaseConstraint(FPhysicsConstraintHandle& InConstraintRef);
+	static PHYSICSCORE_API FPhysicsConstraintHandle CreateSuspension(const FPhysicsActorHandle& InActorRef, const FVector& InLocalFrame);
+	static PHYSICSCORE_API void SetConstraintUserData(const FPhysicsConstraintHandle& InConstraintRef,void* InUserData);
+	static PHYSICSCORE_API void ReleaseConstraint(FPhysicsConstraintHandle& InConstraintRef);
 
-	static FTransform GetLocalPose(const FPhysicsConstraintHandle& InConstraintRef,EConstraintFrame::Type InFrame);
-	static FTransform GetGlobalPose(const FPhysicsConstraintHandle& InConstraintRef,EConstraintFrame::Type InFrame);
-	static FVector GetLocation(const FPhysicsConstraintHandle& InConstraintRef);
-	static void GetForce(const FPhysicsConstraintHandle& InConstraintRef,FVector& OutLinForce,FVector& OutAngForce);
-	static void GetDriveLinearVelocity(const FPhysicsConstraintHandle& InConstraintRef,FVector& OutLinVelocity);
-	static void GetDriveAngularVelocity(const FPhysicsConstraintHandle& InConstraintRef,FVector& OutAngVelocity);
+	static PHYSICSCORE_API FTransform GetLocalPose(const FPhysicsConstraintHandle& InConstraintRef,EConstraintFrame::Type InFrame);
+	static PHYSICSCORE_API FTransform GetGlobalPose(const FPhysicsConstraintHandle& InConstraintRef,EConstraintFrame::Type InFrame);
+	static PHYSICSCORE_API FVector GetLocation(const FPhysicsConstraintHandle& InConstraintRef);
+	static PHYSICSCORE_API void GetForce(const FPhysicsConstraintHandle& InConstraintRef,FVector& OutLinForce,FVector& OutAngForce);
+	static PHYSICSCORE_API void GetDriveLinearVelocity(const FPhysicsConstraintHandle& InConstraintRef,FVector& OutLinVelocity);
+	static PHYSICSCORE_API void GetDriveAngularVelocity(const FPhysicsConstraintHandle& InConstraintRef,FVector& OutAngVelocity);
 
-	static float GetCurrentSwing1(const FPhysicsConstraintHandle& InConstraintRef);
-	static float GetCurrentSwing2(const FPhysicsConstraintHandle& InConstraintRef);
-	static float GetCurrentTwist(const FPhysicsConstraintHandle& InConstraintRef);
+	static PHYSICSCORE_API float GetCurrentSwing1(const FPhysicsConstraintHandle& InConstraintRef);
+	static PHYSICSCORE_API float GetCurrentSwing2(const FPhysicsConstraintHandle& InConstraintRef);
+	static PHYSICSCORE_API float GetCurrentTwist(const FPhysicsConstraintHandle& InConstraintRef);
 
-	static void SetCanVisualize(const FPhysicsConstraintHandle& InConstraintRef,bool bInCanVisualize);
-	static void SetCollisionEnabled(const FPhysicsConstraintHandle& InConstraintRef,bool bInCollisionEnabled);
-	static void SetProjectionEnabled_AssumesLocked(const FPhysicsConstraintHandle& InConstraintRef,bool bInProjectionEnabled,float InLinearAlpha = 1.0f,float InAngularAlpha = 0.0f, float InLinearTolerance = 0.0f, float InAngularToleranceDeg = 0.0f);
-	static void SetShockPropagationEnabled_AssumesLocked(const FPhysicsConstraintHandle& InConstraintRef, bool bInShockPropagationEnabled, float InShockPropagationAlpha);
-	static void SetParentDominates_AssumesLocked(const FPhysicsConstraintHandle& InConstraintRef, bool bInParentDominates);
-	static void SetMassConditioningEnabled_AssumesLocked(const FPhysicsConstraintHandle& InConstraintRef, bool bInMassConditioningEnabled);
-	static void SetBreakForces_AssumesLocked(const FPhysicsConstraintHandle& InConstraintRef,float InLinearBreakForce,float InAngularBreakForce);
-	static void SetPlasticityLimits_AssumesLocked(const FPhysicsConstraintHandle& InConstraintRef, float InLinearPlasticityLimit, float InAngularPlasticityLimit, EConstraintPlasticityType InLinearPlasticityType);
-	static void SetContactTransferScale_AssumesLocked(const FPhysicsConstraintHandle& InConstraintRef, float InContactTransferScale);
-	static void SetLocalPose(const FPhysicsConstraintHandle& InConstraintRef, const FTransform& InPose, EConstraintFrame::Type InFrame);
-
-	static void SetDrivePosition(const FPhysicsConstraintHandle& InConstraintRef,const FVector& InPosition);
-	static void SetDriveOrientation(const FPhysicsConstraintHandle& InConstraintRef,const FQuat& InOrientation);
-	static void SetDriveLinearVelocity(const FPhysicsConstraintHandle& InConstraintRef,const FVector& InLinVelocity);
-	static void SetDriveAngularVelocity(const FPhysicsConstraintHandle& InConstraintRef,const FVector& InAngVelocity);
-
-	static void SetTwistLimit(const FPhysicsConstraintHandle& InConstraintRef,float InLowerLimit,float InUpperLimit,float InContactDistance);
-	static void SetSwingLimit(const FPhysicsConstraintHandle& InConstraintRef,float InYLimit,float InZLimit,float InContactDistance);
-	static void SetLinearLimit(const FPhysicsConstraintHandle& InConstraintRef,float InLimit);
-
-	static bool IsBroken(const FPhysicsConstraintHandle& InConstraintRef);
+	static PHYSICSCORE_API void SetCanVisualize(const FPhysicsConstraintHandle& InConstraintRef,bool bInCanVisualize);
+	static PHYSICSCORE_API void SetCollisionEnabled(const FPhysicsConstraintHandle& InConstraintRef,bool bInCollisionEnabled);
+	static PHYSICSCORE_API void SetProjectionEnabled_AssumesLocked(const FPhysicsConstraintHandle& InConstraintRef,bool bInProjectionEnabled,float InLinearAlpha = 1.0f,float InAngularAlpha = 0.0f, float InLinearTolerance = 0.0f, float InAngularToleranceDeg = 0.0f);
+	static PHYSICSCORE_API void SetShockPropagationEnabled_AssumesLocked(const FPhysicsConstraintHandle& InConstraintRef, bool bInShockPropagationEnabled, float InShockPropagationAlpha);
+	static PHYSICSCORE_API void SetParentDominates_AssumesLocked(const FPhysicsConstraintHandle& InConstraintRef, bool bInParentDominates);
+	static PHYSICSCORE_API void SetMassConditioningEnabled_AssumesLocked(const FPhysicsConstraintHandle& InConstraintRef, bool bInMassConditioningEnabled);
+	static PHYSICSCORE_API void SetBreakForces_AssumesLocked(const FPhysicsConstraintHandle& InConstraintRef,float InLinearBreakForce,float InAngularBreakForce);
+	static PHYSICSCORE_API void SetPlasticityLimits_AssumesLocked(const FPhysicsConstraintHandle& InConstraintRef, float InLinearPlasticityLimit, float InAngularPlasticityLimit, EConstraintPlasticityType InLinearPlasticityType);
+	static PHYSICSCORE_API void SetContactTransferScale_AssumesLocked(const FPhysicsConstraintHandle& InConstraintRef, float InContactTransferScale);
+	static PHYSICSCORE_API void SetLocalPose(const FPhysicsConstraintHandle& InConstraintRef, const FTransform& InPose, EConstraintFrame::Type InFrame);
+
+	static PHYSICSCORE_API void SetDrivePosition(const FPhysicsConstraintHandle& InConstraintRef,const FVector& InPosition);
+	static PHYSICSCORE_API void SetDriveOrientation(const FPhysicsConstraintHandle& InConstraintRef,const FQuat& InOrientation);
+	static PHYSICSCORE_API void SetDriveLinearVelocity(const FPhysicsConstraintHandle& InConstraintRef,const FVector& InLinVelocity);
+	static PHYSICSCORE_API void SetDriveAngularVelocity(const FPhysicsConstraintHandle& InConstraintRef,const FVector& InAngVelocity);
+
+	static PHYSICSCORE_API void SetTwistLimit(const FPhysicsConstraintHandle& InConstraintRef,float InLowerLimit,float InUpperLimit,float InContactDistance);
+	static PHYSICSCORE_API void SetSwingLimit(const FPhysicsConstraintHandle& InConstraintRef,float InYLimit,float InZLimit,float InContactDistance);
+	static PHYSICSCORE_API void SetLinearLimit(const FPhysicsConstraintHandle& InConstraintRef,float InLimit);
+
+	static PHYSICSCORE_API bool IsBroken(const FPhysicsConstraintHandle& InConstraintRef);
 
-	static void SetGeometry(FPhysicsShapeHandle& InShape, TUniquePtr<Chaos::FImplicitObject>&& InGeometry);
+	static PHYSICSCORE_API void SetGeometry(FPhysicsShapeHandle& InShape, TUniquePtr<Chaos::FImplicitObject>&& InGeometry);
 	
-	static FCollisionFilterData GetSimulationFilter(const FPhysicsShapeHandle& InShape);
-	static FCollisionFilterData GetQueryFilter(const FPhysicsShapeHandle& InShape);
-	static bool IsSimulationShape(const FPhysicsShapeHandle& InShape);
-	static bool IsQueryShape(const FPhysicsShapeHandle& InShape);
-	static ECollisionShapeType GetShapeType(const FPhysicsShapeHandle& InShape);
-	static FTransform GetLocalTransform(const FPhysicsShapeHandle& InShape);
+	static PHYSICSCORE_API FCollisionFilterData GetSimulationFilter(const FPhysicsShapeHandle& InShape);
+	static PHYSICSCORE_API FCollisionFilterData GetQueryFilter(const FPhysicsShapeHandle& InShape);
+	static PHYSICSCORE_API bool IsSimulationShape(const FPhysicsShapeHandle& InShape);
+	static PHYSICSCORE_API bool IsQueryShape(const FPhysicsShapeHandle& InShape);
+	static PHYSICSCORE_API ECollisionShapeType GetShapeType(const FPhysicsShapeHandle& InShape);
+	static PHYSICSCORE_API FTransform GetLocalTransform(const FPhysicsShapeHandle& InShape);
 
 
-	static void* GetUserData(const FPhysicsShapeHandle& InShape);
-	static void SetUserData(const FPhysicsShapeHandle& InShape, void* InUserData);
-	static FPhysicsShapeHandle CloneShape(const FPhysicsShapeHandle& InShape);
-	static FPhysicsGeometryCollection_Chaos GetGeometryCollection(const FPhysicsShapeHandle& InShape);
-	static FPhysicsGeometryCollection_Chaos GetGeometryCollection(const FPhysicsGeometry& InShape);
+	static PHYSICSCORE_API void* GetUserData(const FPhysicsShapeHandle& InShape);
+	static PHYSICSCORE_API void SetUserData(const FPhysicsShapeHandle& InShape, void* InUserData);
+	static PHYSICSCORE_API FPhysicsShapeHandle CloneShape(const FPhysicsShapeHandle& InShape);
+	static PHYSICSCORE_API FPhysicsGeometryCollection_Chaos GetGeometryCollection(const FPhysicsShapeHandle& InShape);
+	static PHYSICSCORE_API FPhysicsGeometryCollection_Chaos GetGeometryCollection(const FPhysicsGeometry& InShape);
 
 	// @todo(mlentine): Which of these do we need to support?
 	// Set the mask filter of a shape, which is an extra level of filtering during collision detection / query for extra channels like "Blue Team" and "Red Team"
-	static void SetMaskFilter(const FPhysicsShapeHandle& InShape,FMaskFilter InFilter);
-	static void SetSimulationFilter(const FPhysicsShapeHandle& InShape,const FCollisionFilterData& InFilter);
-	static void SetQueryFilter(const FPhysicsShapeHandle& InShape,const FCollisionFilterData& InFilter);
-	static void SetIsSimulationShape(const FPhysicsShapeHandle& InShape,bool bIsSimShape);
-	static void SetIsProbeShape(const FPhysicsShapeHandle& InShape,bool bIsProbeShape);
-	static void SetIsQueryShape(const FPhysicsShapeHandle& InShape,bool bIsQueryShape);
-	static void SetLocalTransform(const FPhysicsShapeHandle& InShape,const FTransform& NewLocalTransform);
+	static PHYSICSCORE_API void SetMaskFilter(const FPhysicsShapeHandle& InShape,FMaskFilter InFilter);
+	static PHYSICSCORE_API void SetSimulationFilter(const FPhysicsShapeHandle& InShape,const FCollisionFilterData& InFilter);
+	static PHYSICSCORE_API void SetQueryFilter(const FPhysicsShapeHandle& InShape,const FCollisionFilterData& InFilter);
+	static PHYSICSCORE_API void SetIsSimulationShape(const FPhysicsShapeHandle& InShape,bool bIsSimShape);
+	static PHYSICSCORE_API void SetIsProbeShape(const FPhysicsShapeHandle& InShape,bool bIsProbeShape);
+	static PHYSICSCORE_API void SetIsQueryShape(const FPhysicsShapeHandle& InShape,bool bIsQueryShape);
+	static PHYSICSCORE_API void SetLocalTransform(const FPhysicsShapeHandle& InShape,const FTransform& NewLocalTransform);
 };
diff --git a/Engine/Source/Runtime/PhysicsCore/Public/Chaos/ChaosPhysicalMaterial.h b/Engine/Source/Runtime/PhysicsCore/Public/Chaos/ChaosPhysicalMaterial.h
--- a/Engine/Source/Runtime/PhysicsCore/Public/Chaos/ChaosPhysicalMaterial.h
+++ b/Engine/Source/Runtime/PhysicsCore/Public/Chaos/ChaosPhysicalMaterial.h
@@ -18,13 +18,13 @@ namespace Chaos
  * Physical materials are used to define the response of a physical object when 
  * interacting dynamically with the world.
  */
-UCLASS(BlueprintType, Blueprintable, CollapseCategories, HideCategories = Object)
-class PHYSICSCORE_API UChaosPhysicalMaterial : public UObject
+UCLASS(BlueprintType, Blueprintable, CollapseCategories, HideCategories = Object, MinimalAPI)
+class UChaosPhysicalMaterial : public UObject
 {
 	GENERATED_UCLASS_BODY()
 
 	/** Update the values in \p Mat to reflect the values in this \c UObject. */
-	void CopyTo(Chaos::FChaosPhysicsMaterial& Mat) const;
+	PHYSICSCORE_API void CopyTo(Chaos::FChaosPhysicsMaterial& Mat) const;
 
 	//
 	// Surface properties.
diff --git a/Engine/Source/Runtime/PhysicsCore/Public/Chaos/ChaosScene.h b/Engine/Source/Runtime/PhysicsCore/Public/Chaos/ChaosScene.h
--- a/Engine/Source/Runtime/PhysicsCore/Public/Chaos/ChaosScene.h
+++ b/Engine/Source/Runtime/PhysicsCore/Public/Chaos/ChaosScene.h
@@ -79,13 +79,13 @@ struct FChaosSceneSimCallback : public Chaos::TSimCallbackObject<FChaosSceneCall
 /**
 * Low level Chaos scene used when building custom simulations that don't exist in the main world physics scene.
 */
-class PHYSICSCORE_API FChaosScene
+class FChaosScene
 #if WITH_ENGINE
 	: public FGCObject
 #endif
 {
 public:
-	FChaosScene(
+	PHYSICSCORE_API FChaosScene(
 		UObject* OwnerPtr
 		, Chaos::FReal InAsyncDt
 #if CHAOS_DEBUG_NAME
@@ -93,7 +93,7 @@ public:
 #endif
 );
 
-	virtual ~FChaosScene();
+	PHYSICSCORE_API virtual ~FChaosScene();
 
 	/**
 	 * Get the internal Chaos solver object
@@ -102,7 +102,7 @@ public:
 
 #if WITH_ENGINE
 	// FGCObject Interface ///////////////////////////////////////////////////
-	virtual void AddReferencedObjects(FReferenceCollector& Collector) override;
+	PHYSICSCORE_API virtual void AddReferencedObjects(FReferenceCollector& Collector) override;
 	virtual FString GetReferencerName() const
 	{
 		return "FChaosScene";
@@ -110,44 +110,44 @@ public:
 	//////////////////////////////////////////////////////////////////////////
 #endif
 	
-	const Chaos::ISpatialAcceleration<Chaos::FAccelerationStructureHandle, Chaos::FReal, 3>* GetSpacialAcceleration() const;
-	Chaos::ISpatialAcceleration<Chaos::FAccelerationStructureHandle, Chaos::FReal, 3>* GetSpacialAcceleration();
+	PHYSICSCORE_API const Chaos::ISpatialAcceleration<Chaos::FAccelerationStructureHandle, Chaos::FReal, 3>* GetSpacialAcceleration() const;
+	PHYSICSCORE_API Chaos::ISpatialAcceleration<Chaos::FAccelerationStructureHandle, Chaos::FReal, 3>* GetSpacialAcceleration();
 
-	void AddActorsToScene_AssumesLocked(TArray<FPhysicsActorHandle>& InHandles,const bool bImmediate=true);
-	void RemoveActorFromAccelerationStructure(FPhysicsActorHandle Actor);
-	void RemoveActorFromAccelerationStructureImp(Chaos::FGeometryParticle* Particle);
-	void UpdateActorsInAccelerationStructure(const TArrayView<FPhysicsActorHandle>& Actors);
-	void UpdateActorInAccelerationStructure(const FPhysicsActorHandle& Actor);
+	PHYSICSCORE_API void AddActorsToScene_AssumesLocked(TArray<FPhysicsActorHandle>& InHandles,const bool bImmediate=true);
+	PHYSICSCORE_API void RemoveActorFromAccelerationStructure(FPhysicsActorHandle Actor);
+	PHYSICSCORE_API void RemoveActorFromAccelerationStructureImp(Chaos::FGeometryParticle* Particle);
+	PHYSICSCORE_API void UpdateActorsInAccelerationStructure(const TArrayView<FPhysicsActorHandle>& Actors);
+	PHYSICSCORE_API void UpdateActorInAccelerationStructure(const FPhysicsActorHandle& Actor);
 
-	void WaitPhysScenes();
+	PHYSICSCORE_API void WaitPhysScenes();
 
 	/**
 	 * Copies the acceleration structure out of the solver, does no thread safety checking so ensure calls
 	 * to this are made at appropriate sync points if required
 	 */
-	void CopySolverAccelerationStructure();
+	PHYSICSCORE_API void CopySolverAccelerationStructure();
 
 	/**
 	 * Flushes all pending global, task and solver command queues and refreshes the spatial acceleration
 	 * for the scene. Required when querying against a currently non-running scene to ensure the scene
 	 * is correctly represented
 	 */
-	void Flush();
+	PHYSICSCORE_API void Flush();
 #if WITH_EDITOR
-	void AddPieModifiedObject(UObject* InObj);
+	PHYSICSCORE_API void AddPieModifiedObject(UObject* InObj);
 #endif
 
-	void StartFrame();
-	void SetUpForFrame(const FVector* NewGrav,float InDeltaSeconds,float InMinPhysicsDeltaTime,float InMaxPhysicsDeltaTime,float InMaxSubstepDeltaTime,int32 InMaxSubsteps,bool bSubstepping);
-	void EndFrame();
+	PHYSICSCORE_API void StartFrame();
+	PHYSICSCORE_API void SetUpForFrame(const FVector* NewGrav,float InDeltaSeconds,float InMinPhysicsDeltaTime,float InMaxPhysicsDeltaTime,float InMaxSubstepDeltaTime,int32 InMaxSubsteps,bool bSubstepping);
+	PHYSICSCORE_API void EndFrame();
 
 	DECLARE_MULTICAST_DELEGATE_OneParam(FOnPhysScenePostTick,FChaosScene*);
 	FOnPhysScenePostTick OnPhysScenePostTick;
 
-	bool AreAnyTasksPending() const;
-	void BeginDestroy();
-	bool IsCompletionEventComplete() const;
-	FGraphEventArray GetCompletionEvents();
+	PHYSICSCORE_API bool AreAnyTasksPending() const;
+	PHYSICSCORE_API void BeginDestroy();
+	PHYSICSCORE_API bool IsCompletionEventComplete() const;
+	PHYSICSCORE_API FGraphEventArray GetCompletionEvents();
 
 	void SetNetworkDeltaTimeScale(float InDeltaTimeScale) { MNetworkDeltaTimeScale = InDeltaTimeScale; }
 	float GetNetworkDeltaTimeScale() const { return MNetworkDeltaTimeScale; }
@@ -173,7 +173,7 @@ protected:
 
 	//Engine interface BEGIN
 	virtual float OnStartFrame(float InDeltaTime){ return InDeltaTime; }
-	virtual void OnSyncBodies(Chaos::FPhysicsSolverBase* Solver);
+	PHYSICSCORE_API virtual void OnSyncBodies(Chaos::FPhysicsSolverBase* Solver);
 	//Engine interface END
 
 	float MDeltaTime;
@@ -183,7 +183,7 @@ protected:
 
 private:
 
-	void SetGravity(const Chaos::FVec3& Acceleration);
+	PHYSICSCORE_API void SetGravity(const Chaos::FVec3& Acceleration);
 
 	template <typename TSolver>
 	void SyncBodies(TSolver* Solver);
