Add damage modifier to physics material and an option in Geometry Collection asset to opt for it

diff --git a/Engine/Source/Runtime/Experimental/Chaos/Private/PhysicsProxy/GeometryCollectionPhysicsProxy.cpp b/Engine/Source/Runtime/Experimental/Chaos/Private/PhysicsProxy/GeometryCollectionPhysicsProxy.cpp
--- a/Engine/Source/Runtime/Experimental/Chaos/Private/PhysicsProxy/GeometryCollectionPhysicsProxy.cpp
+++ b/Engine/Source/Runtime/Experimental/Chaos/Private/PhysicsProxy/GeometryCollectionPhysicsProxy.cpp
@@ -1534,6 +1534,14 @@ float FGeometryCollectionPhysicsProxy::ComputeUserDefinedDamageThreshold_Interna
 			const float DefaultDamage = NumThresholds > 0 ? Parameters.DamageThreshold[NumThresholds - 1] : 0.f;
 			DamageThreshold = Level < NumThresholds ? Parameters.DamageThreshold[Level] : DefaultDamage;
 		}
+
+		if (Parameters.bUseMaterialDamageModifiers)
+		{
+			if (const Chaos::FChaosPhysicsMaterial* Material = Parameters.PhysicalMaterialHandle.Get())
+			{
+				DamageThreshold *= Material->DamageModifier.DamageThresholdMultiplier;
+			}
+		}
 	}
 
 	return DamageThreshold;
diff --git a/Engine/Source/Runtime/Experimental/Chaos/Public/Chaos/Defines.h b/Engine/Source/Runtime/Experimental/Chaos/Public/Chaos/Defines.h
--- a/Engine/Source/Runtime/Experimental/Chaos/Public/Chaos/Defines.h
+++ b/Engine/Source/Runtime/Experimental/Chaos/Public/Chaos/Defines.h
@@ -63,6 +63,15 @@ namespace Chaos
 		float ShearStrength;
 	};
 
+	struct FChaosPhysicsMaterialDamageModifier
+	{
+		/** 
+		* multiplier for the geometry collection damage thresholds/ internal strain 
+		* this allows for setting up unit damage threshold and use the material to scale them to the desired range of values
+		*/
+		float DamageThresholdMultiplier = 1.0f;
+	};
+
 	class FChaosPhysicsMaterial
 	{
 	public:
@@ -95,6 +104,7 @@ namespace Chaos
 		ECombineMode RestitutionCombineMode;
 
 		FChaosPhysicsMaterialStrength Strength;
+		FChaosPhysicsMaterialDamageModifier DamageModifier;
 
 
 		FChaosPhysicsMaterial()
diff --git a/Engine/Source/Runtime/Experimental/Chaos/Public/GeometryCollection/GeometryCollectionSimulationCoreTypes.h b/Engine/Source/Runtime/Experimental/Chaos/Public/GeometryCollection/GeometryCollectionSimulationCoreTypes.h
--- a/Engine/Source/Runtime/Experimental/Chaos/Public/GeometryCollection/GeometryCollectionSimulationCoreTypes.h
+++ b/Engine/Source/Runtime/Experimental/Chaos/Public/GeometryCollection/GeometryCollectionSimulationCoreTypes.h
@@ -181,6 +181,7 @@ struct FSimulationParameters
 		, MaxClusterLevel(100)
 		, MaxSimulatedLevel(100)
 		, bUseSizeSpecificDamageThresholds(false)
+		, bUseMaterialDamageModifiers(false)
 		, DamageModel(EDamageModelTypeEnum::Chaos_Damage_Model_UserDefined_Damage_Threshold)
 		, DamageEvaluationModel(Chaos::EDamageEvaluationModel::StrainFromDamageThreshold)
 		, DamageThreshold({500000.f, 50000.f, 5000.f})
@@ -237,6 +238,7 @@ struct FSimulationParameters
 		, MaxClusterLevel(Other.MaxClusterLevel)
 		, MaxSimulatedLevel(Other.MaxSimulatedLevel)
 		, bUseSizeSpecificDamageThresholds(Other.bUseSizeSpecificDamageThresholds)
+		, bUseMaterialDamageModifiers(Other.bUseMaterialDamageModifiers)
 		, DamageModel(Other.DamageModel)
 		, DamageEvaluationModel(Other.DamageEvaluationModel)
 		, DamageThreshold(Other.DamageThreshold)
@@ -310,6 +312,7 @@ struct FSimulationParameters
 	int32 MaxClusterLevel;
 	int32 MaxSimulatedLevel;
 	bool bUseSizeSpecificDamageThresholds;
+	bool bUseMaterialDamageModifiers;
 
 	/** this is the user expose damage model, used for creation of the particles */
 	EDamageModelTypeEnum DamageModel; 
diff --git a/Engine/Source/Runtime/Experimental/GeometryCollectionEngine/Private/GeometryCollection/GeometryCollectionComponent.cpp b/Engine/Source/Runtime/Experimental/GeometryCollectionEngine/Private/GeometryCollection/GeometryCollectionComponent.cpp
--- a/Engine/Source/Runtime/Experimental/GeometryCollectionEngine/Private/GeometryCollection/GeometryCollectionComponent.cpp
+++ b/Engine/Source/Runtime/Experimental/GeometryCollectionEngine/Private/GeometryCollection/GeometryCollectionComponent.cpp
@@ -3062,6 +3062,7 @@ void UGeometryCollectionComponent::RegisterAndInitializePhysicsProxy()
 #endif
 		EClusterConnectionTypeEnum ClusterCollectionType = ClusterConnectionType_DEPRECATED;
 		float ConnectionGraphBoundsFilteringMargin = 0;
+		bool bUseMaterialDamageModifiers = false;
 		if (RestCollection)
 		{
 			RestCollection->GetSharedSimulationParams(SimulationParameters.Shared);
@@ -3069,6 +3070,7 @@ void UGeometryCollectionComponent::RegisterAndInitializePhysicsProxy()
 			SimulationParameters.InitialRootIndex = RestCollection->GetRootIndex();
 			ClusterCollectionType = RestCollection->ClusterConnectionType;
 			ConnectionGraphBoundsFilteringMargin = RestCollection->ConnectionGraphBoundsFilteringMargin;
+			bUseMaterialDamageModifiers = RestCollection->bUseMaterialDamageModifiers;
 		}
 		SimulationParameters.Simulating = BodyInstance.bSimulatePhysics;
 		SimulationParameters.EnableClustering = EnableClustering;
@@ -3078,6 +3080,7 @@ void UGeometryCollectionComponent::RegisterAndInitializePhysicsProxy()
 		SimulationParameters.DamageModel = DamageModel;
 		SimulationParameters.DamageEvaluationModel = GetDamageEvaluationModel(DamageModel);
 		SimulationParameters.bUseSizeSpecificDamageThresholds = bUseSizeSpecificDamageThreshold;
+		SimulationParameters.bUseMaterialDamageModifiers = bUseMaterialDamageModifiers;
 		SimulationParameters.DamageThreshold = DamageThreshold;
 		SimulationParameters.bUsePerClusterOnlyDamageThreshold = RestCollection? RestCollection->PerClusterOnlyDamageThreshold: false; 
 		SimulationParameters.ClusterConnectionMethod = (Chaos::FClusterCreationParameters::EConnectionMethod)ClusterCollectionType;
diff --git a/Engine/Source/Runtime/Experimental/GeometryCollectionEngine/Private/GeometryCollection/GeometryCollectionObject.cpp b/Engine/Source/Runtime/Experimental/GeometryCollectionEngine/Private/GeometryCollection/GeometryCollectionObject.cpp
--- a/Engine/Source/Runtime/Experimental/GeometryCollectionEngine/Private/GeometryCollection/GeometryCollectionObject.cpp
+++ b/Engine/Source/Runtime/Experimental/GeometryCollectionEngine/Private/GeometryCollection/GeometryCollectionObject.cpp
@@ -93,6 +93,7 @@ UGeometryCollection::UGeometryCollection(const FObjectInitializer& ObjectInitial
 	, DamageModel(EDamageModelTypeEnum::Chaos_Damage_Model_UserDefined_Damage_Threshold)
 	, DamageThreshold({ 500000.f, 50000.f, 5000.f })
 	, bUseSizeSpecificDamageThreshold(false)
+	, bUseMaterialDamageModifiers(false)
 	, PerClusterOnlyDamageThreshold(false)
 	, ClusterConnectionType(EClusterConnectionTypeEnum::Chaos_MinimalSpanningSubsetDelaunayTriangulation)
 	, ConnectionGraphBoundsFilteringMargin(0)
diff --git a/Engine/Source/Runtime/Experimental/GeometryCollectionEngine/Public/GeometryCollection/GeometryCollectionObject.h b/Engine/Source/Runtime/Experimental/GeometryCollectionEngine/Public/GeometryCollection/GeometryCollectionObject.h
--- a/Engine/Source/Runtime/Experimental/GeometryCollectionEngine/Public/GeometryCollection/GeometryCollectionObject.h
+++ b/Engine/Source/Runtime/Experimental/GeometryCollectionEngine/Public/GeometryCollection/GeometryCollectionObject.h
@@ -529,6 +529,10 @@ public:
 	UPROPERTY(EditAnywhere, Category = "Damage", meta = (EditCondition = "DamageModel == EDamageModelTypeEnum::Chaos_Damage_Model_UserDefined_Damage_Threshold"))
 	bool bUseSizeSpecificDamageThreshold;
 
+	/** When on , use the modifiers on the material to adjust the user defined damage threshold values */
+	UPROPERTY(EditAnywhere, BlueprintReadWrite, Category = "Damage", meta = (EditCondition = "DamageModel == EDamageModelTypeEnum::Chaos_Damage_Model_UserDefined_Damage_Threshold"))
+	bool bUseMaterialDamageModifiers;
+
 	/** compatibility check, when true, only cluster compute damage from parameters and propagate to direct children
 	 *  when false, each child will compute it's damage threshold allowing for more precise and intuitive destruction behavior
 	 */
diff --git a/Engine/Source/Runtime/PhysicsCore/Private/ChaosEngineInterface.cpp b/Engine/Source/Runtime/PhysicsCore/Private/ChaosEngineInterface.cpp
--- a/Engine/Source/Runtime/PhysicsCore/Private/ChaosEngineInterface.cpp
+++ b/Engine/Source/Runtime/PhysicsCore/Private/ChaosEngineInterface.cpp
@@ -242,6 +242,7 @@ void FChaosEngineInterface::UpdateMaterial(FPhysicsMaterialHandle& InHandle,UPhy
 		Material->Strength.TensileStrength = Chaos::MegaPascalToKgPerCmS2(InMaterial->Strength.TensileStrength);
 		Material->Strength.CompressionStrength = Chaos::MegaPascalToKgPerCmS2(InMaterial->Strength.CompressionStrength);
 		Material->Strength.ShearStrength = Chaos::MegaPascalToKgPerCmS2(InMaterial->Strength.ShearStrength);
+		Material->DamageModifier.DamageThresholdMultiplier = InMaterial->DamageModifier.DamageThresholdMultiplier;
 	}
 
 	Chaos::FPhysicalMaterialManager::Get().UpdateMaterial(InHandle);
diff --git a/Engine/Source/Runtime/PhysicsCore/Private/PhysicalMaterial.cpp b/Engine/Source/Runtime/PhysicsCore/Private/PhysicalMaterial.cpp
--- a/Engine/Source/Runtime/PhysicsCore/Private/PhysicalMaterial.cpp
+++ b/Engine/Source/Runtime/PhysicsCore/Private/PhysicalMaterial.cpp
@@ -24,6 +24,11 @@ FPhysicalMaterialStrength::FPhysicalMaterialStrength()
 	ShearStrength = 6;
 }
 
+FPhysicalMaterialDamageModifier::FPhysicalMaterialDamageModifier()
+{
+	DamageThresholdMultiplier = 1.0;
+}
+
 UPhysicalMaterial::UPhysicalMaterial(const FObjectInitializer& ObjectInitializer)
 	: Super(ObjectInitializer)
 {
diff --git a/Engine/Source/Runtime/PhysicsCore/Public/PhysicalMaterials/PhysicalMaterial.h b/Engine/Source/Runtime/PhysicsCore/Public/PhysicalMaterials/PhysicalMaterial.h
--- a/Engine/Source/Runtime/PhysicsCore/Public/PhysicalMaterials/PhysicalMaterial.h
+++ b/Engine/Source/Runtime/PhysicsCore/Public/PhysicalMaterials/PhysicalMaterial.h
@@ -33,25 +33,44 @@ struct FPhysicalMaterialStrength
 	* Tensile strength of the material in MegaPascal ( 10^6 N/m2 )
 	* This amount of tension force per area the material can withstand before it fractures
 	*/
-	UPROPERTY(EditAnywhere, BlueprintReadOnly, Category = "PhysicalMaterial", meta = (ClampMin = 0, ForceUnits = "MPa"))
+	UPROPERTY(EditAnywhere, BlueprintReadOnly, Category = "PhysicalMaterial|Strength", meta = (ClampMin = 0, ForceUnits = "MPa"))
 	float TensileStrength;
 
 	/**
 	* Compression strength of the material in MegaPascal ( 10^6 N/m2 )
 	* This amount of compression force per area the material can withstand before it fractures, crumbles or buckles
 	*/
-	UPROPERTY(EditAnywhere, BlueprintReadOnly, Category = "PhysicalMaterial", meta = (ClampMin = 0, ForceUnits = "MPa"))
+	UPROPERTY(EditAnywhere, BlueprintReadOnly, Category = "PhysicalMaterial|Strength", meta = (ClampMin = 0, ForceUnits = "MPa"))
 	float CompressionStrength;
 
 	/**
 	* Shear strength of the material in MegaPascal ( 10^6 N/m2 )
 	* This amount of shear force per area the material can withstand before it fractures
 	*/
-	UPROPERTY(EditAnywhere, BlueprintReadOnly, Category = "PhysicalMaterial", meta = (ClampMin = 0, ForceUnits = "MPa"))
+	UPROPERTY(EditAnywhere, BlueprintReadOnly, Category = "PhysicalMaterial|Strength", meta = (ClampMin = 0, ForceUnits = "MPa"))
 	float ShearStrength;
 };
 
 
+/**
+ * Damage threshold modifiers, used by the Chaos destruction system
+ */
+USTRUCT(BlueprintType)
+struct FPhysicalMaterialDamageModifier
+{
+	GENERATED_USTRUCT_BODY()
+
+	FPhysicalMaterialDamageModifier();
+
+	/**
+	* Multiplier for the geometry collection damage thresholds/ internal strain
+	* this allows for setting up unit damage threshold and use the material to scale them to the desired range of values
+	* Note that the geometry collection asset needs to opt-in for the material modifer to be able to use it 
+	*/
+	UPROPERTY(EditAnywhere, BlueprintReadOnly, Category = "PhysicalMaterial|DamageModifier", meta = (ClampMin = 0))
+	float DamageThresholdMultiplier;
+};
+
 /**
  * Physical materials are used to define the response of a physical object when interacting dynamically with the world.
  */
@@ -139,6 +158,9 @@ class UPhysicalMaterial : public UObject
 	UPROPERTY(EditAnywhere, BlueprintReadOnly, Category = PhysicalProperties)
 	FPhysicalMaterialStrength Strength;
 
+	UPROPERTY(EditAnywhere, BlueprintReadOnly, Category = PhysicalProperties)
+	FPhysicalMaterialDamageModifier DamageModifier;
+
 public:
 
 	TUniquePtr<FPhysicsMaterialHandle> MaterialHandle;
