Clarify/fix handling of kinematics with the write-back from simulation, in particular when using the cvar SyncKinematicOnGameThread

diff --git a/Engine/Source/Runtime/Experimental/Chaos/Private/PBDRigidsSolver.cpp b/Engine/Source/Runtime/Experimental/Chaos/Private/PBDRigidsSolver.cpp
--- a/Engine/Source/Runtime/Experimental/Chaos/Private/PBDRigidsSolver.cpp
+++ b/Engine/Source/Runtime/Experimental/Chaos/Private/PBDRigidsSolver.cpp
@@ -1637,6 +1637,26 @@ namespace Chaos
 		//MarshallingManager.FreeData_Internal(&PushData);
 	}
 
+	template<typename TRigidParticle>
+	bool ShouldUpdateFromSimulation(const TRigidParticle& InRigidParticle)
+	{
+		if (InRigidParticle.ObjectState() == Chaos::EObjectStateType::Kinematic)
+		{
+			switch (Chaos::SyncKinematicOnGameThread)
+			{
+			case 0:
+				return false;
+			case 1:
+				return true;
+			default:
+				return InRigidParticle.UpdateKinematicFromSimulation();
+			}
+		}
+		// We assume that sleeping/static particles etc won't appear repeatedly (over multiple
+		// frames) in the dirty list, so we can safely return true here without incurring unwanted costs.
+		return true;
+	}
+
 	void FPBDRigidsSolver::ProcessPushedData_Internal(FPushPhysicsData& PushData)
 	{
 		QUICK_SCOPE_CYCLE_COUNTER(ChaosPushData);
@@ -1976,7 +1996,7 @@ namespace Chaos
 						{
 							if(!bIsResim || DirtyParticle.SyncState() == ESyncState::HardDesync)
 							{
-								if (!(SyncKinematicOnGameThread == 0 && DirtyParticle.ObjectState() == EObjectStateType::Kinematic))
+								if (ShouldUpdateFromSimulation(DirtyParticle))
 								{
 									ActiveRigid.AddUnique((FSingleParticlePhysicsProxy*)Proxy);
 								}
diff --git a/Engine/Source/Runtime/Experimental/Chaos/Private/PhysicsProxy/SingleParticlePhysicsProxy.cpp b/Engine/Source/Runtime/Experimental/Chaos/Private/PhysicsProxy/SingleParticlePhysicsProxy.cpp
--- a/Engine/Source/Runtime/Experimental/Chaos/Private/PhysicsProxy/SingleParticlePhysicsProxy.cpp
+++ b/Engine/Source/Runtime/Experimental/Chaos/Private/PhysicsProxy/SingleParticlePhysicsProxy.cpp
@@ -264,30 +264,34 @@ FAutoConsoleVariableRef CVarRenderInterpErrorVelocitySmoothingDuration(TEXT("p.R
 int32 RenderInterpDebugDraw = 0;
 FAutoConsoleVariableRef CVarRenderInterpDebugDraw(TEXT("p.RenderInterp.DebugDraw"), RenderInterpDebugDraw, TEXT("Draw debug lines for physics render interpolation, also needs p.Chaos.DebugDraw.Enabled set"));
 
+bool ShouldUpdateTransformFromSimulation(const Chaos::FPBDRigidParticle& Rigid)
+{
+	if (Rigid.ObjectState() == Chaos::EObjectStateType::Kinematic)
+	{
+		switch (Chaos::SyncKinematicOnGameThread)
+		{
+		case 0:
+			return false;
+		case 1:
+			return true;
+		default:
+			return Rigid.UpdateKinematicFromSimulation();
+		}
+	}
+	return true;
+}
+
 bool FSingleParticlePhysicsProxy::PullFromPhysicsState(const Chaos::FDirtyRigidParticleData& PullData,int32 SolverSyncTimestamp, const Chaos::FDirtyRigidParticleData* NextPullData, const Chaos::FRealSingle* Alpha, const FDirtyRigidParticleReplicationErrorData* Error, const Chaos::FReal AsyncFixedTimeStep)
 {
 	using namespace Chaos;
 	// Move buffered data into the TPBDRigidParticle without triggering invalidation of the physics state.
-	auto Rigid = Particle ? Particle->CastToRigidParticle() : nullptr;
+	Chaos::FPBDRigidParticle* Rigid = Particle ? Particle->CastToRigidParticle() : nullptr;
 	if(Rigid)
 	{
 		// Note that kinematics should either be updated here (following simulation), or when the
 		// kinematic target is set in FChaosEngineInterface::SetKinematicTarget_AssumesLocked If the
 		// logic in one place is changed, it should be checked in the other place too.
-		bool bUpdatePositionFromSimulation = true;
-		if (Rigid->ObjectState() == EObjectStateType::Kinematic)
-		{
-			switch (SyncKinematicOnGameThread)
-			{
-			case 0:
-				bUpdatePositionFromSimulation = false ; break;
-			case 1: 
-				bUpdatePositionFromSimulation = true; break;
-			default:
-				bUpdatePositionFromSimulation = Rigid->UpdateKinematicFromSimulation();
-			}
-		}
-
+		bool bUpdatePositionFromSimulation = ShouldUpdateTransformFromSimulation(*Rigid);
 		const FSingleParticleProxyTimestamp* ProxyTimestamp = PullData.GetTimestamp();
 		
 #if RENDERINTERP_ERRORVELOCITYSMOOTHING
diff --git a/Engine/Source/Runtime/PhysicsCore/Private/ChaosEngineInterface.cpp b/Engine/Source/Runtime/PhysicsCore/Private/ChaosEngineInterface.cpp
--- a/Engine/Source/Runtime/PhysicsCore/Private/ChaosEngineInterface.cpp
+++ b/Engine/Source/Runtime/PhysicsCore/Private/ChaosEngineInterface.cpp
@@ -2172,30 +2172,36 @@ void FChaosEngineInterface::SetGlobalPose_AssumesLocked(const FPhysicsActorHandl
 	Scene->UpdateActorInAccelerationStructure(InActorReference);
 }
 
-void FChaosEngineInterface::SetKinematicTarget_AssumesLocked(const FPhysicsActorHandle& InActorReference,const FTransform& InNewTarget)
+// Match the logic in places that use SyncKinematicOnGameThread (like
+// FSingleParticlePhysicsProxy::PullFromPhysicsState) - to see if that will be updating the
+// position. If not, then we need to do it here. 
+bool ShouldSetKinematicTargetSetGameTransform(const FPhysicsActorHandle& InActorReference)
 {
-	// SetKinematicTarget_AssumesLocked could be called multiple times in one time step
-	const Chaos::FKinematicTarget NewKinematicTarget = Chaos::FKinematicTarget::MakePositionTarget(InNewTarget);
-	InActorReference->GetGameThreadAPI().SetKinematicTarget(NewKinematicTarget);
-
-	// Match the logic in FSingleParticlePhysicsProxy::PullFromPhysicsState - to see if that will be
-	// updating the position. If not, then we need to do it here. Note t
-	bool bUpdatePositionFromSimulation = false;
 	Chaos::FPBDRigidParticle* Rigid = InActorReference->GetRigidParticleUnsafe();
 	if (Rigid && Rigid->ObjectState() == Chaos::EObjectStateType::Kinematic)
 	{
 		switch (Chaos::SyncKinematicOnGameThread)
 		{
 		case 0:
-			bUpdatePositionFromSimulation = false; break;
+			return true;
 		case 1:
-			bUpdatePositionFromSimulation = true; break;
+			return false;
 		default:
-			bUpdatePositionFromSimulation = Rigid->UpdateKinematicFromSimulation();
+			return !Rigid->UpdateKinematicFromSimulation();
 		}
 	}
+	// Historically the game TM gets set through using the kinematic target even if called with a
+	// non-kinematic object, so preserve that behavior.
+	return true;
+}
+
+void FChaosEngineInterface::SetKinematicTarget_AssumesLocked(const FPhysicsActorHandle& InActorReference,const FTransform& InNewTarget)
+{
+	// SetKinematicTarget_AssumesLocked could be called multiple times in one time step
+	const Chaos::FKinematicTarget NewKinematicTarget = Chaos::FKinematicTarget::MakePositionTarget(InNewTarget);
+	InActorReference->GetGameThreadAPI().SetKinematicTarget(NewKinematicTarget);
 
-	if (!bUpdatePositionFromSimulation)
+	if (ShouldSetKinematicTargetSetGameTransform(InActorReference))
 	{
 		// IMPORTANT : we do not invalidate X and R as they will be properly computed using the kinematic target information 
 		InActorReference->GetGameThreadAPI().SetX(InNewTarget.GetLocation(), false); 
