Fixed physics scene pending tasks causing hangs when dispatched during world shutdown

diff --git a/Engine/Source/Runtime/Engine/Private/World.cpp b/Engine/Source/Runtime/Engine/Private/World.cpp
--- a/Engine/Source/Runtime/Engine/Private/World.cpp
+++ b/Engine/Source/Runtime/Engine/Private/World.cpp
@@ -1364,6 +1364,8 @@ bool UWorld::IsReadyForFinishDestroy()
 		if (PhysicsScene != nullptr)
 		{
 			PhysicsScene->KillSafeAsyncTasks();
+			PhysicsScene->WaitSolverTasks();
+
 			if (PhysicsScene->AreAnyTasksPending())
 			{
 				return false;
diff --git a/Engine/Source/Runtime/PhysicsCore/Private/ChaosScene.cpp b/Engine/Source/Runtime/PhysicsCore/Private/ChaosScene.cpp
--- a/Engine/Source/Runtime/PhysicsCore/Private/ChaosScene.cpp
+++ b/Engine/Source/Runtime/PhysicsCore/Private/ChaosScene.cpp
@@ -408,6 +408,15 @@ void FChaosScene::KillSafeAsyncTasks()
 	}
 }
 
+void FChaosScene::WaitSolverTasks()
+{
+	Chaos::FPBDRigidsSolver* Solver = GetSolver();
+	if(Solver)
+	{
+		Solver->WaitOnPendingTasks_External();
+	}
+}
+
 bool FChaosScene::AreAnyTasksPending() const
 {
 	const Chaos::FPBDRigidsSolver* Solver = GetSolver();
diff --git a/Engine/Source/Runtime/PhysicsCore/Public/Chaos/ChaosScene.h b/Engine/Source/Runtime/PhysicsCore/Public/Chaos/ChaosScene.h
--- a/Engine/Source/Runtime/PhysicsCore/Public/Chaos/ChaosScene.h
+++ b/Engine/Source/Runtime/PhysicsCore/Public/Chaos/ChaosScene.h
@@ -145,6 +145,7 @@ public:
 	FOnPhysScenePostTick OnPhysScenePostTick;
 
 	PHYSICSCORE_API void KillSafeAsyncTasks();
+	PHYSICSCORE_API void WaitSolverTasks();
 	PHYSICSCORE_API bool AreAnyTasksPending() const;
 	PHYSICSCORE_API void BeginDestroy();
 	PHYSICSCORE_API bool IsCompletionEventComplete() const;
