Change "CVars" namespace to "PhysicalMaterialCVars" to avoid name conflicts

#changelist validated
#virtualized

diff --git a/Engine/Source/Runtime/PhysicsCore/Private/PhysicalMaterial.cpp b/Engine/Source/Runtime/PhysicsCore/Private/PhysicalMaterial.cpp
--- a/Engine/Source/Runtime/PhysicsCore/Private/PhysicalMaterial.cpp
+++ b/Engine/Source/Runtime/PhysicsCore/Private/PhysicalMaterial.cpp
@@ -12,11 +12,11 @@
 
 #include UE_INLINE_GENERATED_CPP_BY_NAME(PhysicalMaterial)
 
-namespace CVars
+namespace PhysicalMaterialCVars
 {
-	bool bPhysicalMaterial_ShowExperimentalProperties = false;
+	bool bShowExperimentalProperties = false;
 
-	FAutoConsoleVariableRef CVarPhysicalMaterial_ShowExperimentalProperties(TEXT("p.PhysicalMaterial_ShowExperimentalProperties"), bPhysicalMaterial_ShowExperimentalProperties, TEXT(""));
+	FAutoConsoleVariableRef CVarShowExperimentalProperties(TEXT("p.PhysicalMaterial.ShowExperimentalProperties"), bShowExperimentalProperties, TEXT(""));
 }
 
 UDEPRECATED_PhysicalMaterialPropertyBase::UDEPRECATED_PhysicalMaterialPropertyBase(const FObjectInitializer& ObjectInitializer)
diff --git a/Engine/Source/Runtime/PhysicsCore/Public/PhysicalMaterials/PhysicalMaterial.h b/Engine/Source/Runtime/PhysicsCore/Public/PhysicalMaterials/PhysicalMaterial.h
--- a/Engine/Source/Runtime/PhysicsCore/Public/PhysicalMaterials/PhysicalMaterial.h
+++ b/Engine/Source/Runtime/PhysicsCore/Public/PhysicalMaterials/PhysicalMaterial.h
@@ -19,9 +19,9 @@
 
 struct FPropertyChangedEvent;
 
-namespace CVars
+namespace PhysicalMaterialCVars
 {
-	extern bool bPhysicalMaterial_ShowExperimentalProperties;
+	extern bool bShowExperimentalProperties;
 }
 
 
@@ -190,10 +190,10 @@ class UPhysicalMaterial : public UObject
 	// Experimental properties
 	//
 
-	/** Experimental material properties are enabled via the p.PhysicalMaterial_ShowExperimentalProperties console variable.
+	/** Experimental material properties are enabled via the p.PhysicalMaterial.ShowExperimentalProperties console variable.
 		NOTE: These are _experimental_ properties which may change. Use at your own risk! */
 	UPROPERTY(VisibleAnywhere, Transient, Category = "Experimental")
-	bool bShowExperimentalProperties = CVars::bPhysicalMaterial_ShowExperimentalProperties;
+	bool bShowExperimentalProperties = PhysicalMaterialCVars::bShowExperimentalProperties;
 
 	/** For enable soft collision shell thickness mode */
 	UPROPERTY(EditAnywhere, Category = "Experimental|Softness", meta = (EditCondition = "bShowExperimentalProperties"))
