Custom Entities can now be attached to Physics Particles and returned in queries

#[fyi] Benn.Gallagher

diff --git a/Engine/Source/Runtime/Engine/Private/PhysicsEngine/Experimental/ChaosInterfaceWrapper.cpp b/Engine/Source/Runtime/Engine/Private/PhysicsEngine/Experimental/ChaosInterfaceWrapper.cpp
--- a/Engine/Source/Runtime/Engine/Private/PhysicsEngine/Experimental/ChaosInterfaceWrapper.cpp
+++ b/Engine/Source/Runtime/Engine/Private/PhysicsEngine/Experimental/ChaosInterfaceWrapper.cpp
@@ -10,7 +10,21 @@ namespace ChaosInterface
 	FBodyInstance* GetUserData(const Chaos::FGeometryParticle& Actor)
 	{
 		void* UserData = Actor.UserData();
-		return UserData ? FChaosUserData::Get<FBodyInstance>(Actor.UserData()) : nullptr;
+		if (UserData)
+		{
+			FBodyInstance* BodyInstance = FChaosUserData::Get<FBodyInstance>(UserData);
+			if (!BodyInstance)
+			{
+				// Check if we appended a custom entity
+				FChaosUserEntityAppend* ChaosUserEntityAppend = FChaosUserData::Get<FChaosUserEntityAppend>(UserData);
+				if (ChaosUserEntityAppend)
+				{
+					BodyInstance = FChaosUserData::Get<FBodyInstance>(ChaosUserEntityAppend->ChaosUserData);
+				}
+			}
+			return BodyInstance;
+		}
+		return nullptr;
 	}
 
 	UPhysicalMaterial* GetUserData(const Chaos::FChaosPhysicsMaterial& Material)
diff --git a/Engine/Source/Runtime/Engine/Private/PhysicsEngine/Experimental/PhysScene_Chaos.cpp b/Engine/Source/Runtime/Engine/Private/PhysicsEngine/Experimental/PhysScene_Chaos.cpp
--- a/Engine/Source/Runtime/Engine/Private/PhysicsEngine/Experimental/PhysScene_Chaos.cpp
+++ b/Engine/Source/Runtime/Engine/Private/PhysicsEngine/Experimental/PhysScene_Chaos.cpp
@@ -862,7 +862,7 @@ FBodyInstance* FPhysScene_Chaos::GetBodyInstanceFromProxy(const IPhysicsProxyBas
 		if (PhysicsProxy->GetType() == EPhysicsProxyType::SingleParticleProxy)
 		{
 			const Chaos::FRigidBodyHandle_External& RigidBodyHandle = static_cast<const Chaos::FSingleParticlePhysicsProxy*>(PhysicsProxy)->GetGameThreadAPI();
-			BodyInstance = FPhysicsUserData::Get<FBodyInstance>(RigidBodyHandle.UserData());
+			BodyInstance = ChaosInterface::GetUserData(*(static_cast<const Chaos::FSingleParticlePhysicsProxy*>(PhysicsProxy)->GetParticle_LowLevel()));
 		}
 		// found none, let's see if there's an owning component in the scene
 		if (BodyInstance == nullptr)
@@ -2240,7 +2240,7 @@ void FPhysScene_Chaos::OnSyncBodies(Chaos::FPhysicsSolverBase* Solver)
 			SCOPE_CYCLE_COUNTER(STAT_SyncBodiesSingleParticlePhysicsProxy);
 			FPBDRigidParticle* DirtyParticle = Proxy->GetRigidParticleUnsafe();
 
-			if (FBodyInstance* BodyInstance = FPhysicsUserData::Get<FBodyInstance>(DirtyParticle->UserData()))
+			if (FBodyInstance* BodyInstance = ChaosInterface::GetUserData(*(Proxy->GetParticle_LowLevel())))
 			{
 				if (BodyInstance->OwnerComponent.IsValid())
 				{
@@ -2502,7 +2502,7 @@ void FPhysScene_Chaos::ResimNFrames(const int32 NumFramesRequested)
 					UE_LOG(LogChaos,Log,TEXT("Resim had %d desyncs"),DesyncedParticles.Num());
 					for(const FDesyncedParticleInfo& Info : DesyncedParticles)
 					{
-						const FBodyInstance* BI = FPhysicsUserData_Chaos::Get<FBodyInstance>(Info.Particle->UserData());
+						const FBodyInstance* BI = ChaosInterface::GetUserData(Info.Particle);
 						const FBox Bounds = BI->GetBodyBounds();
 						FVector Center,Extents;
 						Bounds.GetCenterAndExtents(Center,Extents);
diff --git a/Engine/Source/Runtime/Experimental/Chaos/Private/Chaos/PhysicsObjectInterface.cpp b/Engine/Source/Runtime/Experimental/Chaos/Private/Chaos/PhysicsObjectInterface.cpp
--- a/Engine/Source/Runtime/Experimental/Chaos/Private/Chaos/PhysicsObjectInterface.cpp
+++ b/Engine/Source/Runtime/Experimental/Chaos/Private/Chaos/PhysicsObjectInterface.cpp
@@ -186,6 +186,33 @@ namespace Chaos
 		return Object ? Object->HasChildren<Id>() : false;
 	}
 
+	template<EThreadContext Id>
+	FChaosUserDefinedEntity* FReadPhysicsObjectInterface<Id>::GetUserDefinedEntity(const FConstPhysicsObjectHandle Object)
+	{
+		if constexpr (Id == EThreadContext::External)
+		{
+			if (!Object)
+			{
+				return nullptr;
+			}
+
+			if (FGeometryParticle* Particle = Object->GetParticle<Id>())
+			{
+				// Do we have a user entity appended?
+				FChaosUserEntityAppend* UserEntityAppend = FChaosUserData::Get<FChaosUserEntityAppend>(Particle->UserData());
+				if (UserEntityAppend)
+				{
+					return UserEntityAppend->UserDefinedEntity;
+				}
+			}
+		}
+		else
+		{
+			ensureMsgf(false, TEXT("User Defined Entities can only be read by the game thread"));
+		}
+		return nullptr;
+	}
+
 	template<EThreadContext Id>
 	int32 FReadPhysicsObjectInterface<Id>::GetClusterHierarchyLevel(const FConstPhysicsObjectHandle Object)
 	{
@@ -817,6 +844,42 @@ namespace Chaos
 		return FAccelerationStructureHandle{InObject->GetParticle<Id>()};
 	}
 
+	template<EThreadContext Id>
+	void FWritePhysicsObjectInterface<Id>::SetUserDefinedEntity(TArrayView<const FPhysicsObjectHandle> InObjects, FChaosUserDefinedEntity* UserDefinedEntity)
+	{
+		if constexpr (Id == EThreadContext::External)
+		{
+			for (const FPhysicsObjectHandle Object : InObjects)
+			{
+				if (!Object)
+				{
+					continue;
+				}
+
+				if (FGeometryParticle* Particle = Object->GetParticle<Id>())
+				{
+					// Do we already have a user entity appended?
+					FChaosUserEntityAppend* UserEntityAppend = FChaosUserData::Get<FChaosUserEntityAppend>(Particle->UserData());
+					if (!UserEntityAppend)
+					{
+						UserEntityAppend = new FChaosUserEntityAppend;
+						UserEntityAppend->ChaosUserData = reinterpret_cast<FChaosUserData*>(Particle->UserData());
+						UserEntityAppend->UserDefinedEntity = UserDefinedEntity;
+						Particle->SetUserData(UserEntityAppend);
+					}
+					else
+					{
+						UserEntityAppend->UserDefinedEntity = UserDefinedEntity; // Overwrite previous used defined entity
+					}
+				}
+			}
+		}
+		else
+		{
+			ensureMsgf(false, TEXT("User Defined Entities can only be set by the game thread"));
+		}		
+	}
+
 	template<EThreadContext Id>
 	void FWritePhysicsObjectInterface<Id>::PutToSleep(TArrayView<const FPhysicsObjectHandle> InObjects)
 	{
diff --git a/Engine/Source/Runtime/Experimental/Chaos/Public/Chaos/ParticleHandle.h b/Engine/Source/Runtime/Experimental/Chaos/Public/Chaos/ParticleHandle.h
--- a/Engine/Source/Runtime/Experimental/Chaos/Public/Chaos/ParticleHandle.h
+++ b/Engine/Source/Runtime/Experimental/Chaos/Public/Chaos/ParticleHandle.h
@@ -2,6 +2,7 @@
 #pragma once
 
 #include "Math/NumericLimits.h"
+#include "Chaos/ChaosUserEntity.h"
 #include "Chaos/ISpatialAcceleration.h"
 #include "Chaos/PBDRigidClusteredParticles.h"
 #include "Chaos/PBDGeometryCollectionParticles.h"
@@ -12,6 +13,7 @@
 #include "ChaosCheck.h"
 #include "Chaos/ChaosDebugDrawDeclares.h"
 #include "ChaosVisualDebugger/ChaosVisualDebuggerTrace.h"
+#include "PhysicsInterfaceTypesCore.h"
 #if CHAOS_DEBUG_DRAW
 #include "Chaos/ChaosDebugDraw.h"
 #endif
@@ -2579,7 +2581,12 @@ public:
 
 	virtual ~TGeometryParticle() //only virtual for easier memory management. Should generally be a static API
 	{
-		SetUniqueIdx(FUniqueIdx{}, false); // Set to an invalid index for dangling handle detection 
+		SetUniqueIdx(FUniqueIdx{}, false); // Set to an invalid index for dangling handle detection
+		FChaosUserEntityAppend* UserEntityAppend = FChaosUserData::Get<FChaosUserEntityAppend>(MUserData);
+		if (UserEntityAppend)
+		{
+			delete UserEntityAppend;
+		}
 	}
 
 	TGeometryParticle(const TGeometryParticle&) = delete;
diff --git a/Engine/Source/Runtime/Experimental/Chaos/Public/Chaos/PhysicsObjectInterface.h b/Engine/Source/Runtime/Experimental/Chaos/Public/Chaos/PhysicsObjectInterface.h
--- a/Engine/Source/Runtime/Experimental/Chaos/Public/Chaos/PhysicsObjectInterface.h
+++ b/Engine/Source/Runtime/Experimental/Chaos/Public/Chaos/PhysicsObjectInterface.h
@@ -16,6 +16,7 @@
 #include "PhysicsObjectInterface.generated.h"
 
 class FChaosScene;
+class FChaosUserDefinedEntity;
 class IPhysicsProxyBase;
 
 USTRUCT(BlueprintType)
@@ -61,6 +62,7 @@ namespace Chaos
 	public:
 		CHAOS_API FPhysicsObjectHandle GetRootObject(const FConstPhysicsObjectHandle Object);
 		CHAOS_API bool HasChildren(const FConstPhysicsObjectHandle Object);
+		CHAOS_API FChaosUserDefinedEntity* GetUserDefinedEntity(const FConstPhysicsObjectHandle Object);
 		CHAOS_API int32 GetClusterHierarchyLevel(const FConstPhysicsObjectHandle Object);
 
 		CHAOS_API FTransform GetTransform(const FConstPhysicsObjectHandle Object);
@@ -123,6 +125,7 @@ namespace Chaos
 	class FWritePhysicsObjectInterface: public FReadPhysicsObjectInterface<Id>
 	{
 	public:
+		CHAOS_API void SetUserDefinedEntity(TArrayView<const FPhysicsObjectHandle> InObjects, FChaosUserDefinedEntity* UserDefinedEntity);
 		CHAOS_API void PutToSleep(TArrayView<const FPhysicsObjectHandle> InObjects);
 		CHAOS_API void WakeUp(TArrayView<const FPhysicsObjectHandle> InObjects);
 		CHAOS_API void ForceKinematic(TArrayView<const FPhysicsObjectHandle> InObjects);
diff --git a/Engine/Source/Runtime/PhysicsCore/Public/Chaos/ChaosUserEntity.h b/Engine/Source/Runtime/PhysicsCore/Public/Chaos/ChaosUserEntity.h
new file mode 100644
--- /dev/null
+++ b/Engine/Source/Runtime/PhysicsCore/Public/Chaos/ChaosUserEntity.h
@@ -0,0 +1,25 @@
+// Copyright Epic Games, Inc. All Rights Reserved.
+
+#pragma once
+#include "PhysicsInterfaceTypesCore.h"
+
+class FChaosUserDefinedEntity
+{
+public:
+	FChaosUserDefinedEntity() = delete;
+	FChaosUserDefinedEntity(FName InEntityTypeName) : EntityTypeName(InEntityTypeName) {};
+
+	FName GetEntityTypeName() {
+		return EntityTypeName;
+	}
+private:	
+	FName EntityTypeName;
+};
+
+// This is used to add your own Unreal Engine agnostic user entities to Chaos Physics Results
+struct FChaosUserEntityAppend : public FChaosUserData
+{
+	FChaosUserEntityAppend() { Type = EChaosUserDataType::ChaosUserEntity; Payload = this; };
+	FChaosUserData* ChaosUserData; //  Chaos data to access the physics body properties
+	FChaosUserDefinedEntity* UserDefinedEntity;
+};
diff --git a/Engine/Source/Runtime/PhysicsCore/Public/PhysicsInterfaceTypesCore.h b/Engine/Source/Runtime/PhysicsCore/Public/PhysicsInterfaceTypesCore.h
--- a/Engine/Source/Runtime/PhysicsCore/Public/PhysicsInterfaceTypesCore.h
+++ b/Engine/Source/Runtime/PhysicsCore/Public/PhysicsInterfaceTypesCore.h
@@ -174,6 +174,8 @@ struct FKShapeElem;
 /** Forward declarations */
 struct FKShapeElem;
 struct FCustomChaosPayload;
+struct FPhysicsObject;
+struct FChaosUserEntityAppend;
 
 namespace EChaosUserDataType
 {
@@ -186,6 +188,8 @@ namespace EChaosUserDataType
 		ConstraintInstance,
 		PrimitiveComponent,
 		AggShape,
+		PhysicsObject,    // This will replace BodyInstances in the future
+		ChaosUserEntity,  // This is used for adding custom user entities (unknown to UE)
 		CustomPayload,	//This is intended for plugins
 	};
 };
@@ -205,6 +209,8 @@ public:
 	FChaosUserData(FConstraintInstance* InPayload)		:Type(EChaosUserDataType::ConstraintInstance), Payload(InPayload) {}
 	FChaosUserData(UPrimitiveComponent* InPayload)		:Type(EChaosUserDataType::PrimitiveComponent), Payload(InPayload) {}
 	FChaosUserData(FKShapeElem* InPayload)				:Type(EChaosUserDataType::AggShape), Payload(InPayload) {}
+	FChaosUserData(FPhysicsObject* InPayload)			:Type(EChaosUserDataType::PhysicsObject), Payload(InPayload) {}
+	FChaosUserData(FChaosUserEntityAppend* InPayload)	:Type(EChaosUserDataType::ChaosUserEntity), Payload(InPayload) {}
 	FChaosUserData(FCustomChaosPayload* InPayload)		:Type(EChaosUserDataType::CustomPayload), Payload(InPayload) {}
 	
 	template <class T> static T* Get(void* UserData);
@@ -223,6 +229,8 @@ template <> FORCEINLINE FPhysScene* FChaosUserData::Get(void* UserData)				{ if
 template <> FORCEINLINE FConstraintInstance* FChaosUserData::Get(void* UserData)	{ if (!UserData || ((FChaosUserData*)UserData)->Type != EChaosUserDataType::ConstraintInstance) { return nullptr; } return (FConstraintInstance*)((FChaosUserData*)UserData)->Payload; }
 template <> FORCEINLINE UPrimitiveComponent* FChaosUserData::Get(void* UserData)	{ if (!UserData || ((FChaosUserData*)UserData)->Type != EChaosUserDataType::PrimitiveComponent) { return nullptr; } return (UPrimitiveComponent*)((FChaosUserData*)UserData)->Payload; }
 template <> FORCEINLINE FKShapeElem* FChaosUserData::Get(void* UserData)	{ if (!UserData || ((FChaosUserData*)UserData)->Type != EChaosUserDataType::AggShape) { return nullptr; } return (FKShapeElem*)((FChaosUserData*)UserData)->Payload; }
+template <> FORCEINLINE FPhysicsObject* FChaosUserData::Get(void* UserData) { if (!UserData || ((FChaosUserData*)UserData)->Type != EChaosUserDataType::PhysicsObject) { return nullptr; } return (FPhysicsObject*)((FChaosUserData*)UserData)->Payload; }
+template <> FORCEINLINE FChaosUserEntityAppend* FChaosUserData::Get(void* UserData) { if (!UserData || ((FChaosUserData*)UserData)->Type != EChaosUserDataType::ChaosUserEntity) { return nullptr; } return (FChaosUserEntityAppend*)((FChaosUserData*)UserData)->Payload; }
 template <> FORCEINLINE FCustomChaosPayload* FChaosUserData::Get(void* UserData) { if (!UserData || ((FChaosUserData*)UserData)->Type != EChaosUserDataType::CustomPayload) { return nullptr; } return (FCustomChaosPayload*)((FChaosUserData*)UserData)->Payload; }
 
 template <> FORCEINLINE void FChaosUserData::Set(void* UserData, FBodyInstance* Payload)			{ check(UserData); ((FChaosUserData*)UserData)->Type = EChaosUserDataType::BodyInstance; ((FChaosUserData*)UserData)->Payload = Payload; }
@@ -231,6 +239,8 @@ template <> FORCEINLINE void FChaosUserData::Set(void* UserData, FPhysScene* Pay
 template <> FORCEINLINE void FChaosUserData::Set(void* UserData, FConstraintInstance* Payload)		{ check(UserData); ((FChaosUserData*)UserData)->Type = EChaosUserDataType::ConstraintInstance; ((FChaosUserData*)UserData)->Payload = Payload; }
 template <> FORCEINLINE void FChaosUserData::Set(void* UserData, UPrimitiveComponent* Payload)		{ check(UserData); ((FChaosUserData*)UserData)->Type = EChaosUserDataType::PrimitiveComponent; ((FChaosUserData*)UserData)->Payload = Payload; }
 template <> FORCEINLINE void FChaosUserData::Set(void* UserData, FKShapeElem* Payload)	{ check(UserData); ((FChaosUserData*)UserData)->Type = EChaosUserDataType::AggShape; ((FChaosUserData*)UserData)->Payload = Payload; }
+template <> FORCEINLINE void FChaosUserData::Set(void* UserData, FPhysicsObject* Payload) { check(UserData); ((FChaosUserData*)UserData)->Type = EChaosUserDataType::PhysicsObject; ((FChaosUserData*)UserData)->Payload = Payload; }
+template <> FORCEINLINE void FChaosUserData::Set(void* UserData, FChaosUserEntityAppend* Payload) { check(UserData); ((FChaosUserData*)UserData)->Type = EChaosUserDataType::ChaosUserEntity; ((FChaosUserData*)UserData)->Payload = Payload; }
 template <> FORCEINLINE void FChaosUserData::Set(void* UserData, FCustomChaosPayload* Payload) { check(UserData); ((FChaosUserData*)UserData)->Type = EChaosUserDataType::CustomPayload; ((FChaosUserData*)UserData)->Payload = Payload; }
 
 struct FChaosFilterData
