Chaos User Entity now includes a virtual destructor to support user entity destruct from a base class pointer
#[fyi] Benn.Gallagher

diff --git a/Engine/Source/Runtime/PhysicsCore/Public/Chaos/ChaosUserEntity.h b/Engine/Source/Runtime/PhysicsCore/Public/Chaos/ChaosUserEntity.h
--- a/Engine/Source/Runtime/PhysicsCore/Public/Chaos/ChaosUserEntity.h
+++ b/Engine/Source/Runtime/PhysicsCore/Public/Chaos/ChaosUserEntity.h
@@ -8,6 +8,7 @@ class FChaosUserDefinedEntity
 public:
 	FChaosUserDefinedEntity() = delete;
 	FChaosUserDefinedEntity(FName InEntityTypeName) : EntityTypeName(InEntityTypeName) {};
+	virtual ~FChaosUserDefinedEntity() {};
 
 	FName GetEntityTypeName() {
 		return EntityTypeName;
