TSQVisitor needs to initialize HitFaceNormal in all constructors

diff --git a/Engine/Source/Runtime/PhysicsCore/Public/SQVisitor.h b/Engine/Source/Runtime/PhysicsCore/Public/SQVisitor.h
--- a/Engine/Source/Runtime/PhysicsCore/Public/SQVisitor.h
+++ b/Engine/Source/Runtime/PhysicsCore/Public/SQVisitor.h
@@ -50,6 +50,7 @@ struct TSQVisitor : public Chaos::ISpatialVisitor<TPayload, Chaos::FReal>
 		, OutputFlags(InOutputFlags)
 		, bAnyHit(false)
 		, DebugParams(InDebugParams)
+		, HitFaceNormal(Chaos::FVec3::ZeroVector)
 		, HitBuffer(InHitBuffer)
 		, QueryFilterData(InQueryFilterData)
 		, QueryFilterDataConcrete(C2UFilterData(QueryFilterData.data))
