Chaos - fix joint angle reporting on game thread

diff --git a/Engine/Source/Runtime/PhysicsCore/Private/ChaosEngineInterface.cpp b/Engine/Source/Runtime/PhysicsCore/Private/ChaosEngineInterface.cpp
--- a/Engine/Source/Runtime/PhysicsCore/Private/ChaosEngineInterface.cpp
+++ b/Engine/Source/Runtime/PhysicsCore/Private/ChaosEngineInterface.cpp
@@ -1518,14 +1518,14 @@ FTransform FChaosEngineInterface::GetGlobalPose(const FPhysicsConstraintHandle&
 			{
 				if (Chaos::FGeometryParticle* Particle = GetParticleFromProxy(BasePairs[0]))
 				{
-					return FTransform(Particle->R(), Particle->X()) * M[0];
+					return M[0] * FTransform(Particle->R(), Particle->X());
 				}
 			}
 			else if (InFrame == EConstraintFrame::Frame2)
 			{
 				if (Chaos::FGeometryParticle* Particle = GetParticleFromProxy(BasePairs[1]))
 				{
-					return FTransform(Particle->R(), Particle->X()) * M[1];
+					return M[1] * FTransform(Particle->R(), Particle->X());
 				}
 			}
 		}
@@ -1587,19 +1587,26 @@ void FChaosEngineInterface::GetDriveAngularVelocity(const FPhysicsConstraintHand
 	}
 }
 
+FTransform GetConstraintBodiesRelativeTransform(const FPhysicsConstraintHandle& InConstraintRef)
+{
+	const FTransform ChildTransform = FChaosEngineInterface::GetGlobalPose(InConstraintRef, EConstraintFrame::Frame1);
+	const FTransform ParentTransform = FChaosEngineInterface::GetGlobalPose(InConstraintRef, EConstraintFrame::Frame2);
+	return ChildTransform.GetRelativeTransform(ParentTransform);
+}
+
 float FChaosEngineInterface::GetCurrentSwing1(const FPhysicsConstraintHandle& InConstraintRef)
 {
-	return GetLocalPose(InConstraintRef,EConstraintFrame::Frame2).GetRotation().Euler().X;
+	return GetConstraintBodiesRelativeTransform(InConstraintRef).GetRotation().GetTwistAngle(Chaos::FJointConstants::Swing1Axis());
 }
 
 float FChaosEngineInterface::GetCurrentSwing2(const FPhysicsConstraintHandle& InConstraintRef)
 {
-	return GetLocalPose(InConstraintRef,EConstraintFrame::Frame2).GetRotation().Euler().Y;
+	return GetConstraintBodiesRelativeTransform(InConstraintRef).GetRotation().GetTwistAngle(Chaos::FJointConstants::Swing2Axis());
 }
 
 float FChaosEngineInterface::GetCurrentTwist(const FPhysicsConstraintHandle& InConstraintRef)
 {
-	return GetLocalPose(InConstraintRef,EConstraintFrame::Frame2).GetRotation().Euler().Z;
+	return GetConstraintBodiesRelativeTransform(InConstraintRef).GetRotation().GetTwistAngle(Chaos::FJointConstants::TwistAxis());
 }
 
 void FChaosEngineInterface::SetCanVisualize(const FPhysicsConstraintHandle& InConstraintRef,bool bInCanVisualize)
