Add FBodyInstance::GetBodyBoundsLocal

diff --git a/Engine/Source/Runtime/Engine/Classes/PhysicsEngine/BodyInstance.h b/Engine/Source/Runtime/Engine/Classes/PhysicsEngine/BodyInstance.h
--- a/Engine/Source/Runtime/Engine/Classes/PhysicsEngine/BodyInstance.h
+++ b/Engine/Source/Runtime/Engine/Classes/PhysicsEngine/BodyInstance.h
@@ -821,8 +821,10 @@ public:
 
 	/** Returns the body's mass */
 	ENGINE_API float GetBodyMass() const;
-	/** Return bounds of physics representation */
+	/** Return bounds of physics representation in world space */
 	ENGINE_API FBox GetBodyBounds() const;
+	/** Return bounds of physics representation in world space */
+	ENGINE_API FBox GetBodyBoundsLocal() const;
 	/** Return the body's inertia tensor. This is returned in local mass space */
 	ENGINE_API FVector GetBodyInertiaTensor() const;
 
diff --git a/Engine/Source/Runtime/Engine/Private/PhysicsEngine/BodyInstance.cpp b/Engine/Source/Runtime/Engine/Private/PhysicsEngine/BodyInstance.cpp
--- a/Engine/Source/Runtime/Engine/Private/PhysicsEngine/BodyInstance.cpp
+++ b/Engine/Source/Runtime/Engine/Private/PhysicsEngine/BodyInstance.cpp
@@ -2807,9 +2807,21 @@ FBox FBodyInstance::GetBodyBounds() const
 	FBox OutBox(EForceInit::ForceInitToZero);
 
 	FPhysicsCommand::ExecuteRead(ActorHandle, [&](const FPhysicsActorHandle& Actor)
-	{
-		OutBox = FPhysicsInterface::GetBounds_AssumesLocked(Actor);
-	});
+		{
+			OutBox = FPhysicsInterface::GetBounds_AssumesLocked(Actor);
+		});
+
+	return OutBox;
+}
+
+FBox FBodyInstance::GetBodyBoundsLocal() const
+{
+	FBox OutBox(EForceInit::ForceInitToZero);
+
+	FPhysicsCommand::ExecuteRead(ActorHandle, [&](const FPhysicsActorHandle& Actor)
+		{
+			OutBox = FPhysicsInterface::GetBoundsLocal_AssumesLocked(Actor);
+		});
 
 	return OutBox;
 }
diff --git a/Engine/Source/Runtime/PhysicsCore/Private/ChaosEngineInterface.cpp b/Engine/Source/Runtime/PhysicsCore/Private/ChaosEngineInterface.cpp
--- a/Engine/Source/Runtime/PhysicsCore/Private/ChaosEngineInterface.cpp
+++ b/Engine/Source/Runtime/PhysicsCore/Private/ChaosEngineInterface.cpp
@@ -765,6 +765,22 @@ FBox FChaosEngineInterface::GetBounds_AssumesLocked(const FPhysicsActorHandle& I
 	return FBox(EForceInit::ForceInitToZero);
 }
 
+FBox FChaosEngineInterface::GetBoundsLocal_AssumesLocked(const FPhysicsActorHandle& InActorReference)
+{
+	using namespace Chaos;
+	const Chaos::FRigidBodyHandle_External& Body_External = InActorReference->GetGameThreadAPI();
+	if (const FImplicitObjectRef Geometry = Body_External.GetGeometry())
+	{
+		if (Geometry->HasBoundingBox())
+		{
+			const FAABB3 LocalBounds = Geometry->BoundingBox();
+			return FBox(LocalBounds.Min(), LocalBounds.Max());
+		}
+	}
+
+	return FBox(EForceInit::ForceInitToZero);
+}
+
 void FChaosEngineInterface::SetLinearDamping_AssumesLocked(const FPhysicsActorHandle& InActorReference,float InDrag)
 {
 	if(ensure(FChaosEngineInterface::IsValid(InActorReference)))
diff --git a/Engine/Source/Runtime/PhysicsCore/Public/Chaos/ChaosEngineInterface.h b/Engine/Source/Runtime/PhysicsCore/Public/Chaos/ChaosEngineInterface.h
--- a/Engine/Source/Runtime/PhysicsCore/Public/Chaos/ChaosEngineInterface.h
+++ b/Engine/Source/Runtime/PhysicsCore/Public/Chaos/ChaosEngineInterface.h
@@ -438,6 +438,7 @@ public:
 	static PHYSICSCORE_API FVector GetLocalInertiaTensor_AssumesLocked(const FPhysicsActorHandle& InActorReference);
 	static PHYSICSCORE_API FBox GetBounds_AssumesLocked(const FPhysicsActorHandle& InActorReference);
 	static PHYSICSCORE_API FBox GetBounds_AssumesLocked(const FPhysicsActorHandle& InActorReference, const FTransform& InTransform);
+	static PHYSICSCORE_API FBox GetBoundsLocal_AssumesLocked(const FPhysicsActorHandle& InActorReference);
 
 	static PHYSICSCORE_API void SetLinearDamping_AssumesLocked(const FPhysicsActorHandle& InActorReference,float InDamping);
 	static PHYSICSCORE_API void SetAngularDamping_AssumesLocked(const FPhysicsActorHandle& InActorReference,float InDamping);
