Fix use of HitResult and OverlapResult in LatentActions when using verse

diff --git a/Engine/Source/Runtime/Engine/Classes/Engine/HitResult.h b/Engine/Source/Runtime/Engine/Classes/Engine/HitResult.h
--- a/Engine/Source/Runtime/Engine/Classes/Engine/HitResult.h
+++ b/Engine/Source/Runtime/Engine/Classes/Engine/HitResult.h
@@ -130,6 +130,8 @@ struct FHitResult
 	UPROPERTY()
 	TWeakObjectPtr<UPrimitiveComponent> Component;
 
+	/** The object that owns the PhysicsObject. This is used to determine if the PhysicsObject is still valid when not owned by an Actor */
+	TWeakObjectPtr<UObject> PhysicsObjectOwner;
 
 	/** PhysicsObjects hit by the query. Not exposed to blueprints for the time being */
 	Chaos::FPhysicsObjectHandle PhysicsObject;
diff --git a/Engine/Source/Runtime/Engine/Classes/Engine/OverlapResult.h b/Engine/Source/Runtime/Engine/Classes/Engine/OverlapResult.h
--- a/Engine/Source/Runtime/Engine/Classes/Engine/OverlapResult.h
+++ b/Engine/Source/Runtime/Engine/Classes/Engine/OverlapResult.h
@@ -31,6 +31,9 @@ struct FOverlapResult
 	/** Utility to return the Component that was hit */
 	ENGINE_API UPrimitiveComponent* GetComponent() const;
 
+	/** The object that owns the PhysicsObject. This is used to determine if the PhysicsObject is still valid when not owned by an Actor */
+	TWeakObjectPtr<UObject> PhysicsObjectOwner;
+
 	/** PhysicsObjects hit by the query. Not exposed to blueprints for the time being */
 	Chaos::FPhysicsObjectHandle PhysicsObject;
 
diff --git a/Engine/Source/Runtime/Engine/Private/Collision/CollisionConversions.cpp b/Engine/Source/Runtime/Engine/Private/Collision/CollisionConversions.cpp
--- a/Engine/Source/Runtime/Engine/Private/Collision/CollisionConversions.cpp
+++ b/Engine/Source/Runtime/Engine/Private/Collision/CollisionConversions.cpp
@@ -184,6 +184,15 @@ static void SetHitResultFromShapeAndFaceIndex(const FPhysicsShape& Shape,  const
 		}
 	}
 
+	if (OutResult.PhysicsObject != nullptr)
+	{
+		OutResult.PhysicsObjectOwner = OwningComponent;
+		if (FChaosUserEntityAppend* ChaosUserEntityAppend = FChaosUserData::Get<FChaosUserEntityAppend>(Actor.UserData()))
+		{
+			OutResult.PhysicsObjectOwner = ChaosUserEntityAppend->GetOwnerObject();
+		}
+	}
+
 	OutResult.PhysMaterial = nullptr;
 
 	// Grab actor/component
@@ -630,6 +639,15 @@ void ConvertQueryOverlap(const FPhysicsShape& Shape, const FPhysicsActor& Actor,
 			OutOverlap.ItemIndex = OwnerComponent->bMultiBodyOverlap ? BodyInst->InstanceBodyIndex : INDEX_NONE;
 		}
 		OutOverlap.PhysicsObject = BodyInst->ActorHandle ? BodyInst->ActorHandle->GetPhysicsObject() : nullptr;
+
+		if (OutOverlap.PhysicsObject != nullptr)
+		{
+			OutOverlap.PhysicsObjectOwner = OutOverlap.Component;
+			if (FChaosUserEntityAppend* ChaosUserEntityAppend = FChaosUserData::Get<FChaosUserEntityAppend>(Actor.UserData()))
+			{
+				OutOverlap.PhysicsObjectOwner = ChaosUserEntityAppend->GetOwnerObject();
+			}
+		}
 	}
 	else
 	{
@@ -816,6 +834,7 @@ FHitResult ConvertOverlapToHitResult(const FOverlapResult& Overlap)
 	Hit.Item = Overlap.ItemIndex;
 	Hit.Component = Overlap.Component;
 	Hit.PhysicsObject = Overlap.PhysicsObject;
+	Hit.PhysicsObjectOwner = Overlap.PhysicsObjectOwner;
 	Hit.HitObjectHandle = Overlap.OverlapObjectHandle;
 	return Hit;
 }
diff --git a/Engine/Source/Runtime/PhysicsCore/Public/Chaos/ChaosUserEntity.h b/Engine/Source/Runtime/PhysicsCore/Public/Chaos/ChaosUserEntity.h
--- a/Engine/Source/Runtime/PhysicsCore/Public/Chaos/ChaosUserEntity.h
+++ b/Engine/Source/Runtime/PhysicsCore/Public/Chaos/ChaosUserEntity.h
@@ -2,6 +2,8 @@
 
 #pragma once
 #include "PhysicsInterfaceTypesCore.h"
+#include "UObject/WeakObjectPtr.h"
+#include "UObject/WeakObjectPtrTemplates.h"
 
 class FChaosUserDefinedEntity
 {
@@ -13,6 +15,10 @@ public:
 	FName GetEntityTypeName() {
 		return EntityTypeName;
 	}
+
+	// Get the UObject (usually a verse::component) that owns this FChaosUserDefinedEntity
+	virtual TWeakObjectPtr<UObject> GetOwnerObject() = 0;
+
 private:	
 	FName EntityTypeName;
 };
@@ -20,7 +26,28 @@ private:
 // This is used to add your own Unreal Engine agnostic user entities to Chaos Physics Results
 struct FChaosUserEntityAppend : public FChaosUserData
 {
-	FChaosUserEntityAppend() { Type = EChaosUserDataType::ChaosUserEntity; Payload = this; };
+	FChaosUserEntityAppend()
+		: FChaosUserData()
+		, ChaosUserData(nullptr)
+		, UserDefinedEntity(nullptr)
+	{ 
+		Type = EChaosUserDataType::ChaosUserEntity;
+		Payload = this;
+	}
+
+	// Get the UObject that controls the lifetime of the physics objects
+	// that will be referencing this UserData object. It will be stored
+	// as a weak pointer in hit/overlap results to verify that the raw
+	// PhysicsObject pointer they hold is valid.
+	TWeakObjectPtr<UObject> GetOwnerObject()
+	{
+		if (UserDefinedEntity != nullptr)
+		{
+			return UserDefinedEntity->GetOwnerObject();
+		}
+		return TWeakObjectPtr<UObject>();
+	}
+
 	FChaosUserData* ChaosUserData; //  Chaos data to access the physics body properties
 	FChaosUserDefinedEntity* UserDefinedEntity;
 };
