Correct spelling for "shrink"  deformation type EConstraintPlasticityType

diff --git a/Engine/Source/Runtime/PhysicsCore/Public/Chaos/ChaosEngineInterface.h b/Engine/Source/Runtime/PhysicsCore/Public/Chaos/ChaosEngineInterface.h
--- a/Engine/Source/Runtime/PhysicsCore/Public/Chaos/ChaosEngineInterface.h
+++ b/Engine/Source/Runtime/PhysicsCore/Public/Chaos/ChaosEngineInterface.h
@@ -173,7 +173,7 @@ enum EConstraintPlasticityType : int
 	/** */
 	CCPT_Free	UMETA(DisplayName = "Free"),
 	/** */
-	CCPT_Shrink	UMETA(DisplayName = "Shirnk"),
+	CCPT_Shrink	UMETA(DisplayName = "Shrink"),
 	/** */
 	CCPT_Grow	UMETA(DisplayName = "Grow"),
 
