- Add Inflate method to FCollisionShape and use it in diverse shape components
- Add SetShape to be able to store a shape type and extent and then recreate the Collision shape from it
- Add IsInside method for TCapsuleShape similar to TSphere

diff --git a/Engine/Source/Runtime/Core/Public/Math/CapsuleShape.h b/Engine/Source/Runtime/Core/Public/Math/CapsuleShape.h
--- a/Engine/Source/Runtime/Core/Public/Math/CapsuleShape.h
+++ b/Engine/Source/Runtime/Core/Public/Math/CapsuleShape.h
@@ -58,6 +58,15 @@ public:
 		: TCapsuleShape<T>(TVector<T>(From.Center), (T)From.Radius, TVector<T>(From.Orientation), (T)From.Length)
 	{
 	}
+
+	bool IsInside(const TVector<T>& In, T Tolerance = UE_KINDA_SMALL_NUMBER) const
+	{
+		const TVector<T> HalfLengthAxis = 0.5 * Length * Orientation;
+		const TVector<T> SegmentStart = Center - HalfLengthAxis;
+		const TVector<T> SegmentEnd = Center + HalfLengthAxis;
+		const TVector<T> Projection = FMath::ClosestPointOnSegment(In, SegmentStart, SegmentEnd);
+		return (Projection - In).SizeSquared() <= FMath::Square(Radius + Tolerance);
+	}
 };
 
 }	// namespace UE::Math
diff --git a/Engine/Source/Runtime/Engine/Private/Components/BoxComponent.cpp b/Engine/Source/Runtime/Engine/Private/Components/BoxComponent.cpp
--- a/Engine/Source/Runtime/Engine/Private/Components/BoxComponent.cpp
+++ b/Engine/Source/Runtime/Engine/Private/Components/BoxComponent.cpp
@@ -249,13 +249,7 @@ FPrimitiveSceneProxy* UBoxComponent::CreateSceneProxy()
 
 FCollisionShape UBoxComponent::GetCollisionShape(float Inflation) const
 {
-	FVector Extent = GetScaledBoxExtent() + Inflation;
-	if (Inflation < 0.f)
-	{
-		// Don't shrink below zero size.
-		Extent = Extent.ComponentMax(FVector::ZeroVector);
-	}
-
-	return FCollisionShape::MakeBox(Extent);
+	const FVector Extent = GetScaledBoxExtent();
+	return FCollisionShape::MakeBox(Extent).Inflate(Inflation);
 }
 
diff --git a/Engine/Source/Runtime/Engine/Private/Components/CapsuleComponent.cpp b/Engine/Source/Runtime/Engine/Private/Components/CapsuleComponent.cpp
--- a/Engine/Source/Runtime/Engine/Private/Components/CapsuleComponent.cpp
+++ b/Engine/Source/Runtime/Engine/Private/Components/CapsuleComponent.cpp
@@ -238,9 +238,9 @@ bool UCapsuleComponent::IsZeroExtent() const
 
 FCollisionShape UCapsuleComponent::GetCollisionShape(float Inflation) const
 {
-	const float Radius = FMath::Max(0.f, GetScaledCapsuleRadius() + Inflation);
-	const float HalfHeight = FMath::Max(0.f, GetScaledCapsuleHalfHeight() + Inflation);
-	return FCollisionShape::MakeCapsule(Radius, HalfHeight);
+	const float Radius = GetScaledCapsuleRadius();
+	const float HalfHeight = GetScaledCapsuleHalfHeight();
+	return FCollisionShape::MakeCapsule(Radius, HalfHeight).Inflate(Inflation);
 }
 
 bool UCapsuleComponent::AreSymmetricRotations(const FQuat& A, const FQuat& B, const FVector& Scale3D) const
diff --git a/Engine/Source/Runtime/Engine/Private/Components/PrimitiveComponent.cpp b/Engine/Source/Runtime/Engine/Private/Components/PrimitiveComponent.cpp
--- a/Engine/Source/Runtime/Engine/Private/Components/PrimitiveComponent.cpp
+++ b/Engine/Source/Runtime/Engine/Private/Components/PrimitiveComponent.cpp
@@ -2696,14 +2696,7 @@ bool UPrimitiveComponent::ShouldComponentIgnoreHitResult(FHitResult const& TestH
 FCollisionShape UPrimitiveComponent::GetCollisionShape(float Inflation) const
 {
 	// This is intended to be overridden by shape classes, so this is a simple, large bounding shape.
-	FVector Extent = Bounds.BoxExtent + Inflation;
-	if (Inflation < 0.f)
-	{
-		// Don't shrink below zero size.
-		Extent = Extent.ComponentMax(FVector::ZeroVector);
-	}
-
-	return FCollisionShape::MakeBox(Extent);
+	return FCollisionShape::MakeBox(Bounds.BoxExtent).Inflate(Inflation);
 }
 
 
diff --git a/Engine/Source/Runtime/Engine/Private/Components/SphereComponent.cpp b/Engine/Source/Runtime/Engine/Private/Components/SphereComponent.cpp
--- a/Engine/Source/Runtime/Engine/Private/Components/SphereComponent.cpp
+++ b/Engine/Source/Runtime/Engine/Private/Components/SphereComponent.cpp
@@ -199,8 +199,8 @@ FPrimitiveSceneProxy* USphereComponent::CreateSceneProxy()
 
 FCollisionShape USphereComponent::GetCollisionShape(float Inflation) const
 {
-	const float Radius = FMath::Max(0.f, GetScaledSphereRadius() + Inflation);
-	return FCollisionShape::MakeSphere(Radius);
+	const float Radius = GetScaledSphereRadius();
+	return FCollisionShape::MakeSphere(Radius).Inflate(Inflation);
 }
 
 bool USphereComponent::AreSymmetricRotations(const FQuat& A, const FQuat& B, const FVector& Scale3D) const
diff --git a/Engine/Source/Runtime/PhysicsCore/Public/CollisionShape.h b/Engine/Source/Runtime/PhysicsCore/Public/CollisionShape.h
--- a/Engine/Source/Runtime/PhysicsCore/Public/CollisionShape.h
+++ b/Engine/Source/Runtime/PhysicsCore/Public/CollisionShape.h
@@ -109,6 +109,31 @@ struct FCollisionShape
 		Capsule.HalfHeight = Extent.Z;
 	}
 
+	/** Utility function to set a shape from a type and extent */
+	void SetShape(const ECollisionShape::Type InShapeType, const FVector& Extent)
+	{
+		switch (InShapeType)
+		{
+		case ECollisionShape::Box:
+			{
+				SetBox(FVector3f(Extent));
+			}
+			break;
+		case ECollisionShape::Sphere:
+			{
+				SetSphere(static_cast<float>(Extent[0]));
+			}
+			break;
+		case ECollisionShape::Capsule:
+			{
+				SetCapsule(FVector3f(Extent));
+			}
+			break;
+		default:
+			ShapeType = InShapeType;
+		}
+	}
+
 	/** Return true if nearly zero. If so, it will back out and use line trace instead */
 	bool IsNearlyZero() const
 	{
@@ -118,7 +143,7 @@ struct FCollisionShape
 		{
 			return (Box.HalfExtentX <= FCollisionShape::MinBoxExtent() && Box.HalfExtentY <= FCollisionShape::MinBoxExtent() && Box.HalfExtentZ <= FCollisionShape::MinBoxExtent());
 		}
-		case  ECollisionShape::Sphere:
+		case ECollisionShape::Sphere:
 		{
 			return (Sphere.Radius <= FCollisionShape::MinSphereRadius());
 		}
@@ -141,7 +166,7 @@ struct FCollisionShape
 		{
 			return FVector(Box.HalfExtentX, Box.HalfExtentY, Box.HalfExtentZ);
 		}
-		case  ECollisionShape::Sphere:
+		case ECollisionShape::Sphere:
 		{
 			return FVector(Sphere.Radius, Sphere.Radius, Sphere.Radius);
 		}
@@ -186,6 +211,69 @@ struct FCollisionShape
 		return Capsule.HalfHeight;
 	}
 
+	FCollisionShape Inflate(const float Inflation) const
+	{
+		FCollisionShape InflatedShape;
+
+		if (Inflation == 0.f)
+		{
+			InflatedShape = *this;
+		}
+		else
+		{
+			// Don't shrink below zero size.
+			switch (ShapeType)
+			{
+			case ECollisionShape::Box:
+				{
+					const FVector3f InflatedExtent = FVector3f(Box.HalfExtentX + Inflation, Box.HalfExtentY + Inflation, Box.HalfExtentZ + Inflation).ComponentMax(FVector3f::ZeroVector);
+					InflatedShape.SetBox(InflatedExtent);
+				}
+				break;
+			case ECollisionShape::Sphere:
+				{
+					InflatedShape.SetSphere(FMath::Max(Sphere.Radius + Inflation, 0));
+				}
+				break;
+			case ECollisionShape::Capsule:
+				{
+					InflatedShape.SetCapsule(FMath::Max(Capsule.Radius + Inflation, 0), FMath::Max(Capsule.HalfHeight + Inflation, 0));
+				}
+				break;
+			default:
+				{
+					// do not inflate for unsupported shapes
+					InflatedShape = *this;
+				}
+			}
+		}
+
+		return InflatedShape;
+	}
+
+	FORCEINLINE FString ToString() const
+	{
+		switch (ShapeType)
+		{
+		case ECollisionShape::Box:
+			{
+				return FString::Printf(TEXT("Box=(X=%3.3f Y=%3.3f Z=%3.3f)"), Box.HalfExtentX, Box.HalfExtentY, Box.HalfExtentZ);
+			}
+		case ECollisionShape::Sphere:
+			{
+				return FString::Printf(TEXT("Sphere=(Radius=%3.3f)"), Sphere.Radius);
+			}
+		case ECollisionShape::Capsule:
+			{
+				return FString::Printf(TEXT("Capsule=(Radius=%3.3f HalfHeight=%3.3f)"), Capsule.Radius, Capsule.HalfHeight);
+			}
+		default:
+			{
+				return TEXT("Line");
+			}
+		}
+	}
+
 	/** Used by engine in multiple places. Since LineShape doesn't need any dimension, declare once and used by all codes. */
 	static struct FCollisionShape LineShape;
 
diff --git a/Engine/Source/Runtime/UMG/Private/Components/WidgetComponent.cpp b/Engine/Source/Runtime/UMG/Private/Components/WidgetComponent.cpp
--- a/Engine/Source/Runtime/UMG/Private/Components/WidgetComponent.cpp
+++ b/Engine/Source/Runtime/UMG/Private/Components/WidgetComponent.cpp
@@ -902,15 +902,8 @@ FCollisionShape UWidgetComponent::GetCollisionShape(float Inflation) const
 {
 	if ( Space != EWidgetSpace::Screen )
 	{
-		FVector BoxHalfExtent = ( FVector(0.01f, CurrentDrawSize.X * 0.5f, CurrentDrawSize.Y * 0.5f) * GetComponentTransform().GetScale3D() ) + Inflation;
-
-		if ( Inflation < 0.0f )
-		{
-			// Don't shrink below zero size.
-			BoxHalfExtent = BoxHalfExtent.ComponentMax(FVector::ZeroVector);
-		}
-
-		return FCollisionShape::MakeBox(BoxHalfExtent);
+		const FVector BoxHalfExtent = ( FVector(0.01f, CurrentDrawSize.X * 0.5f, CurrentDrawSize.Y * 0.5f) * GetComponentTransform().GetScale3D() );
+		return FCollisionShape::MakeBox(BoxHalfExtent).Inflate(Inflation);
 	}
 	else
 	{
