Fix for IsAnyRigidBodyAwake returns true when a dynamic Actor is not simulating physics

IsSleeping in FChaosEngineInterface checks whether the ObjectState is equal to the sleeping state. This check is too specific because there are other states that are also considered not awake that are not the sleeping state. In this case, setting "Simulate Physics" to false will set the object state to Kinematic which will cause this function and subsequently IsAnyRigidBodyAwake to incorrectly return false. This change alters the code to check against the Dynamic state instead, which is the only state that can be considered Awake.

diff --git a/Engine/Source/Runtime/PhysicsCore/Private/ChaosEngineInterface.cpp b/Engine/Source/Runtime/PhysicsCore/Private/ChaosEngineInterface.cpp
--- a/Engine/Source/Runtime/PhysicsCore/Private/ChaosEngineInterface.cpp
+++ b/Engine/Source/Runtime/PhysicsCore/Private/ChaosEngineInterface.cpp
@@ -427,7 +427,7 @@ bool FChaosEngineInterface::IsKinematic_AssumesLocked(const FPhysicsActorHandle&
 
 bool FChaosEngineInterface::IsSleeping(const FPhysicsActorHandle& InActorReference)
 {
-	return InActorReference->GetGameThreadAPI().ObjectState() == Chaos::EObjectStateType::Sleeping;
+	return InActorReference->GetGameThreadAPI().ObjectState() != Chaos::EObjectStateType::Dynamic;
 }
 
 bool FChaosEngineInterface::IsCcdEnabled(const FPhysicsActorHandle& InActorReference)
