Fixes for SwitchUnhandledEnumerator and UnsafeTypeCast warnings in Core/Engine/PhysicsCore/SlateCore headers.
These changes should not affect behavior as they're just expressing the existing behavior more explicitly.

diff --git a/Engine/Source/Runtime/Core/Public/Templates/TypeHash.h b/Engine/Source/Runtime/Core/Public/Templates/TypeHash.h
--- a/Engine/Source/Runtime/Core/Public/Templates/TypeHash.h
+++ b/Engine/Source/Runtime/Core/Public/Templates/TypeHash.h
@@ -188,4 +188,4 @@ template <typename T>
 
 // Use this when inside type that has GetTypeHash() (no in-parameters) implemented. It makes GetTypeHash dispatch in global namespace
 template <typename T>
-[[nodiscard]] FORCEINLINE uint32 GetTypeHashHelper(const T& V) { return GetTypeHash(V); }
+[[nodiscard]] FORCEINLINE uint32 GetTypeHashHelper(const T& V) { return static_cast<uint32>(GetTypeHash(V)); }
diff --git a/Engine/Source/Runtime/Engine/Classes/Camera/PlayerCameraManager.h b/Engine/Source/Runtime/Engine/Classes/Camera/PlayerCameraManager.h
--- a/Engine/Source/Runtime/Engine/Classes/Camera/PlayerCameraManager.h
+++ b/Engine/Source/Runtime/Engine/Classes/Camera/PlayerCameraManager.h
@@ -154,6 +154,8 @@ public:
 		case VTBlend_EaseInOut: return FMath::InterpEaseInOut(0.f, 1.f, TimePct, BlendExp); 
 		case VTBlend_EaseIn: return FMath::Lerp(0.f, 1.f, FMath::Pow(TimePct, BlendExp)); 
 		case VTBlend_EaseOut: return FMath::Lerp(0.f, 1.f, FMath::Pow(TimePct, (FMath::IsNearlyZero(BlendExp) ? 1.f : (1.f / BlendExp))));
+		case VTBlend_PreBlended:
+		case VTBlend_MAX:
 		default:
 			break;
 		}
diff --git a/Engine/Source/Runtime/Engine/Classes/EdGraph/EdGraphNode.h b/Engine/Source/Runtime/Engine/Classes/EdGraph/EdGraphNode.h
--- a/Engine/Source/Runtime/Engine/Classes/EdGraph/EdGraphNode.h
+++ b/Engine/Source/Runtime/Engine/Classes/EdGraph/EdGraphNode.h
@@ -108,6 +108,7 @@ inline const TCHAR* const LexToString(const EEdGraphPinDirection State)
 		return TEXT("Input");
 	case EEdGraphPinDirection::EGPD_Output:
 		return TEXT("Output");
+	case EEdGraphPinDirection::EGPD_MAX:
 	default:
 		break;
 	}
diff --git a/Engine/Source/Runtime/Engine/Classes/Engine/EngineTypes.h b/Engine/Source/Runtime/Engine/Classes/Engine/EngineTypes.h
--- a/Engine/Source/Runtime/Engine/Classes/Engine/EngineTypes.h
+++ b/Engine/Source/Runtime/Engine/Classes/Engine/EngineTypes.h
@@ -3413,6 +3413,7 @@ public:
 				return 2.0f;
 			}
 
+			case WalkableSlope_Max:
 			default:
 			{
 				return InWalkableFloorZ;
diff --git a/Engine/Source/Runtime/Engine/Classes/Engine/StaticMesh.h b/Engine/Source/Runtime/Engine/Classes/Engine/StaticMesh.h
--- a/Engine/Source/Runtime/Engine/Classes/Engine/StaticMesh.h
+++ b/Engine/Source/Runtime/Engine/Classes/Engine/StaticMesh.h
@@ -116,6 +116,7 @@ inline const TCHAR* ToString(EStaticMeshAsyncProperties Value)
 			return TEXT("HiResSourceModel");
 		case EStaticMeshAsyncProperties::UseLegacyTangentScaling:
 			return TEXT("UseLegacyTangentScaling");
+		case EStaticMeshAsyncProperties::All:
 		default: 
 			check(false); 
 			return TEXT("Unknown");
diff --git a/Engine/Source/Runtime/PhysicsCore/Public/CollisionShape.h b/Engine/Source/Runtime/PhysicsCore/Public/CollisionShape.h
--- a/Engine/Source/Runtime/PhysicsCore/Public/CollisionShape.h
+++ b/Engine/Source/Runtime/PhysicsCore/Public/CollisionShape.h
@@ -129,6 +129,7 @@ struct FCollisionShape
 				SetCapsule(FVector3f(Extent));
 			}
 			break;
+		case ECollisionShape::Line:
 		default:
 			ShapeType = InShapeType;
 		}
@@ -240,6 +241,7 @@ struct FCollisionShape
 					InflatedShape.SetCapsule(FMath::Max(Capsule.Radius + Inflation, 0), FMath::Max(Capsule.HalfHeight + Inflation, 0));
 				}
 				break;
+			case ECollisionShape::Line:
 			default:
 				{
 					// do not inflate for unsupported shapes
@@ -267,6 +269,7 @@ struct FCollisionShape
 			{
 				return FString::Printf(TEXT("Capsule=(Radius=%3.3f HalfHeight=%3.3f)"), Capsule.Radius, Capsule.HalfHeight);
 			}
+		case ECollisionShape::Line:
 		default:
 			{
 				return TEXT("Line");
diff --git a/Engine/Source/Runtime/SlateCore/Public/Fonts/FontRasterizationMode.h b/Engine/Source/Runtime/SlateCore/Public/Fonts/FontRasterizationMode.h
--- a/Engine/Source/Runtime/SlateCore/Public/Fonts/FontRasterizationMode.h
+++ b/Engine/Source/Runtime/SlateCore/Public/Fonts/FontRasterizationMode.h
@@ -31,6 +31,7 @@ inline bool IsSdfFontRasterizationMode(EFontRasterizationMode Mode)
 		case EFontRasterizationMode::Sdf:
 		case EFontRasterizationMode::SdfApproximation:
 			return true;
+		case EFontRasterizationMode::Bitmap:
 		default:
 			return false;
 	}
diff --git a/Engine/Source/Runtime/SlateCore/Public/Styling/SlateColor.h b/Engine/Source/Runtime/SlateCore/Public/Styling/SlateColor.h
--- a/Engine/Source/Runtime/SlateCore/Public/Styling/SlateColor.h
+++ b/Engine/Source/Runtime/SlateCore/Public/Styling/SlateColor.h
@@ -110,6 +110,7 @@ public:
 		{
 		default:
 		case ESlateColorStylingMode::UseColor_Foreground:
+		case ESlateColorStylingMode::UseColor_UseStyle:
 			return InWidgetStyle.GetForegroundColor();
 			break;
 
