PhysicalMaterial: Editor linker fix for subclasses of UPhysicalMaterial in other modules.
#rnx

diff --git a/Engine/Source/Runtime/PhysicsCore/Public/PhysicalMaterials/PhysicalMaterial.h b/Engine/Source/Runtime/PhysicsCore/Public/PhysicalMaterials/PhysicalMaterial.h
--- a/Engine/Source/Runtime/PhysicsCore/Public/PhysicalMaterials/PhysicalMaterial.h
+++ b/Engine/Source/Runtime/PhysicsCore/Public/PhysicalMaterials/PhysicalMaterial.h
@@ -229,7 +229,7 @@ public:
 	PHYSICSCORE_API virtual void FinishDestroy() override;
 	//~ End UObject Interface
 
-	virtual void Serialize(FArchive& Ar) override;
+	PHYSICSCORE_API virtual void Serialize(FArchive& Ar) override;
 
 	/** Get the physics-interface derived version of this material */
 	PHYSICSCORE_API FPhysicsMaterialHandle& GetPhysicsMaterial();
