Fixed a bug with the physics queue handler where capsule information was not properly communicated. Two classes defined the height of a capsule in different ways.

diff --git a/Engine/Source/Runtime/Engine/Private/Collision/SceneQuery.cpp b/Engine/Source/Runtime/Engine/Private/Collision/SceneQuery.cpp
--- a/Engine/Source/Runtime/Engine/Private/Collision/SceneQuery.cpp
+++ b/Engine/Source/Runtime/Engine/Private/Collision/SceneQuery.cpp
@@ -197,7 +197,12 @@ bool BuildQueryShape(const FPhysicsGeometryCollection& GeometryCollection, Chaos
 		case ECollisionShapeType::Capsule:
 		{
 			const Chaos::FCapsule& Capsule = GeometryCollection.GetCapsuleGeometry();
-			OutQueryShape.CollisionShape = FCollisionShape::MakeCapsule(Capsule.GetRadiusf(), Capsule.GetHeightf() * 0.5f);
+			const float Radius = Capsule.GetRadiusf();
+			const float HalfHeight = Capsule.GetHeightf() * 0.5f;
+			// Note: FCollisionShape and FCapsule disagree on what the (half) height are. 
+			// FCapsule defines height as the distance between the sphere centers while FCollision defines it as the distance between the full top to bottom.
+			const float FullHalfHeight = HalfHeight + Radius;
+			OutQueryShape.CollisionShape = FCollisionShape::MakeCapsule(Radius, FullHalfHeight);
 			return true;
 		}
 		case ECollisionShapeType::Convex:
diff --git a/Engine/Source/Runtime/PhysicsCore/Public/CollisionShape.h b/Engine/Source/Runtime/PhysicsCore/Public/CollisionShape.h
--- a/Engine/Source/Runtime/PhysicsCore/Public/CollisionShape.h
+++ b/Engine/Source/Runtime/PhysicsCore/Public/CollisionShape.h
@@ -93,7 +93,7 @@ struct FCollisionShape
 		Sphere.Radius = Radius;
 	}
 
-	/** Utility function to set Capsule with Radius and Half Height */
+	/** Utility function to set Capsule with Radius and Half Height. Note: This is the full half-height (needs to include the sphere radius) */
 	void SetCapsule(const float Radius, const float HalfHeight)
 	{
 		ShapeType = ECollisionShape::Capsule;
@@ -206,7 +206,7 @@ struct FCollisionShape
 		return Capsule.Radius;
 	}
 
-	/** Utility function to get Capsule Half Height */
+	/** Utility function to get Capsule Half Height. Note: This is the full half height which includes the sphere radius */
 	const float GetCapsuleHalfHeight() const
 	{
 		return Capsule.HalfHeight;
@@ -304,7 +304,7 @@ struct FCollisionShape
 		return SphereShape;
 	}
 
-	/** Static utility function to make a capsule */
+	/** Static utility function to make a capsule. Note: This is the full half-height (needs to include the sphere radius) */
 	static FCollisionShape MakeCapsule(const float CapsuleRadius, const float CapsuleHalfHeight)
 	{
 		FCollisionShape CapsuleShape;
