Fixed Engine module includes in Chaos module causing build failures in some configurations.

diff --git a/Engine/Source/Runtime/Experimental/Chaos/Private/PBDRigidsSolver.cpp b/Engine/Source/Runtime/Experimental/Chaos/Private/PBDRigidsSolver.cpp
--- a/Engine/Source/Runtime/Experimental/Chaos/Private/PBDRigidsSolver.cpp
+++ b/Engine/Source/Runtime/Experimental/Chaos/Private/PBDRigidsSolver.cpp
@@ -38,7 +38,6 @@
 #include "ProfilingDebugging/CountersTrace.h"
 #include "ProfilingDebugging/CsvProfiler.h"
 #include "ChaosLog.h"
-#include "Physics/Experimental/PhysScene_Chaos.h"
 
 //PRAGMA_DISABLE_OPTIMIZATION
 
diff --git a/Engine/Source/Runtime/Experimental/Chaos/Private/PhysicsProxy/SingleParticlePhysicsProxy.cpp b/Engine/Source/Runtime/Experimental/Chaos/Private/PhysicsProxy/SingleParticlePhysicsProxy.cpp
--- a/Engine/Source/Runtime/Experimental/Chaos/Private/PhysicsProxy/SingleParticlePhysicsProxy.cpp
+++ b/Engine/Source/Runtime/Experimental/Chaos/Private/PhysicsProxy/SingleParticlePhysicsProxy.cpp
@@ -17,7 +17,6 @@
 #include "Chaos/PullPhysicsDataImp.h"
 #include "RewindData.h"
 #include "Chaos/DebugDrawQueue.h"
-#include "Physics/Experimental/PhysScene_Chaos.h"
 
 namespace Chaos
 {
diff --git a/Engine/Source/Runtime/PhysicsCore/Private/ChaosEngineInterface.cpp b/Engine/Source/Runtime/PhysicsCore/Private/ChaosEngineInterface.cpp
--- a/Engine/Source/Runtime/PhysicsCore/Private/ChaosEngineInterface.cpp
+++ b/Engine/Source/Runtime/PhysicsCore/Private/ChaosEngineInterface.cpp
@@ -12,7 +12,6 @@
 #include "Chaos/Serialization/SerializationUtils.h"
 #include "Chaos/Serialization/SerializedDataBuffer.h"
 #include "Chaos/Serialization/SolverSerializer.h"
-#include "Physics/Experimental/PhysInterface_Chaos.h"
 
 FPhysicsDelegatesCore::FOnUpdatePhysXMaterial FPhysicsDelegatesCore::OnUpdatePhysXMaterial;
 
