Fix AddForce not triggering OnComponentWake.

The call to SetObjectState in AddForce does not set the bAllowEvents parameter, resulting in the callbacks not being called. Setting this argument fixes the issue.

diff --git a/Engine/Source/Runtime/PhysicsCore/Private/ChaosEngineInterface.cpp b/Engine/Source/Runtime/PhysicsCore/Private/ChaosEngineInterface.cpp
--- a/Engine/Source/Runtime/PhysicsCore/Private/ChaosEngineInterface.cpp
+++ b/Engine/Source/Runtime/PhysicsCore/Private/ChaosEngineInterface.cpp
@@ -903,7 +903,11 @@ struct FChaosStateOps
 	static void AddForce(BodyHandleType& BodyHandle, const FVector& Force, bool bAllowSubstepping, bool bAccelChange)
 	{
 		Chaos::EObjectStateType ObjectState = BodyHandle.ObjectState();
-		BodyHandle.SetObjectState(Chaos::EObjectStateType::Dynamic);
+		
+		if (!Force.IsZero())
+		{
+			BodyHandle.SetObjectState(Chaos::EObjectStateType::Dynamic, true);
+		}
 
 		if (bAccelChange)
 		{
