Chaos - minor fixes

diff --git a/Engine/Source/Runtime/Experimental/Chaos/Private/PBDRigidsSolver.cpp b/Engine/Source/Runtime/Experimental/Chaos/Private/PBDRigidsSolver.cpp
--- a/Engine/Source/Runtime/Experimental/Chaos/Private/PBDRigidsSolver.cpp
+++ b/Engine/Source/Runtime/Experimental/Chaos/Private/PBDRigidsSolver.cpp
@@ -651,6 +651,7 @@ namespace Chaos
 
 	FPBDRigidsSolver::FPBDRigidsSolver(const EMultiBufferMode BufferingModeIn, UObject* InOwner, FReal InAsyncDt)
 		: Super(BufferingModeIn, BufferingModeIn == EMultiBufferMode::Single ? EThreadingModeTemp::SingleThread : EThreadingModeTemp::TaskGraph, InOwner, InAsyncDt)
+		, PhysSceneHack(nullptr)
 		, CurrentFrame(0)
 		, bHasFloor(true)
 		, bIsFloorAnalytic(false)
diff --git a/Engine/Source/Runtime/PhysicsCore/Private/ChaosEngineInterface.cpp b/Engine/Source/Runtime/PhysicsCore/Private/ChaosEngineInterface.cpp
--- a/Engine/Source/Runtime/PhysicsCore/Private/ChaosEngineInterface.cpp
+++ b/Engine/Source/Runtime/PhysicsCore/Private/ChaosEngineInterface.cpp
@@ -2381,8 +2381,10 @@ void FChaosEngineInterface::SetGlobalPose_AssumesLocked(const FPhysicsActorHandl
 
 	Body_External.UpdateShapeBounds();
 
-	FChaosScene* Scene = GetCurrentScene(InActorReference);
-	Scene->UpdateActorInAccelerationStructure(InActorReference);
+	if (FChaosScene* Scene = GetCurrentScene(InActorReference))
+	{
+		Scene->UpdateActorInAccelerationStructure(InActorReference);
+	}
 }
 
 // Match the logic in places that use SyncKinematicOnGameThread (like
