[Chaos] Added debug names for particles created as a kinematic end point in DOF joint constraints.

This is needed because in the Chaos Visual Debugger they will appear as unnamed particles, and is difficult to understand where they come from

diff --git a/Engine/Source/Runtime/Experimental/Chaos/Private/Chaos/PhysicsObjectInterface.cpp b/Engine/Source/Runtime/Experimental/Chaos/Private/Chaos/PhysicsObjectInterface.cpp
--- a/Engine/Source/Runtime/Experimental/Chaos/Private/Chaos/PhysicsObjectInterface.cpp
+++ b/Engine/Source/Runtime/Experimental/Chaos/Private/Chaos/PhysicsObjectInterface.cpp
@@ -1319,6 +1319,19 @@ namespace Chaos
 		return Object->GetBodyName();
 	}
 
+#if CHAOS_DEBUG_NAME
+	const FString& FPhysicsObjectInterface::GetDebugName(const FConstPhysicsObjectHandle Object)
+	{
+		TThreadParticle<EThreadContext::External>* Particle = Object ? Object->GetParticle<EThreadContext::External>() : nullptr;
+		if (!Object)
+		{
+			return Private::EmptyParticleName;
+		}
+
+		return Particle->GetDebugName();
+	}
+#endif
+
 	void FPhysicsObjectInterface::SetId(const FPhysicsObjectHandle Object, int32 InId)
 	{
 		if (!Object)
diff --git a/Engine/Source/Runtime/Experimental/Chaos/Public/Chaos/PhysicsObjectInterface.h b/Engine/Source/Runtime/Experimental/Chaos/Public/Chaos/PhysicsObjectInterface.h
--- a/Engine/Source/Runtime/Experimental/Chaos/Public/Chaos/PhysicsObjectInterface.h
+++ b/Engine/Source/Runtime/Experimental/Chaos/Public/Chaos/PhysicsObjectInterface.h
@@ -187,6 +187,10 @@ namespace Chaos
 		static CHAOS_API void SetName(const FPhysicsObjectHandle Object, const FName& InName);
 		static CHAOS_API FName GetName(const FConstPhysicsObjectHandle Object);
 
+#if CHAOS_DEBUG_NAME
+		/** Returns the debug name of the underlying particle */
+		static CHAOS_API const FString& GetDebugName(const FConstPhysicsObjectHandle Object);
+#endif
 		static CHAOS_API void SetId(const FPhysicsObjectHandle Object, int32 InId);
 		static CHAOS_API int32 GetId(const FConstPhysicsObjectHandle Object);
 
diff --git a/Engine/Source/Runtime/PhysicsCore/Private/ChaosEngineInterface.cpp b/Engine/Source/Runtime/PhysicsCore/Private/ChaosEngineInterface.cpp
--- a/Engine/Source/Runtime/PhysicsCore/Private/ChaosEngineInterface.cpp
+++ b/Engine/Source/Runtime/PhysicsCore/Private/ChaosEngineInterface.cpp
@@ -12,6 +12,8 @@
 #include "Chaos/Serialization/SerializationUtils.h"
 #include "Chaos/Serialization/SerializedDataBuffer.h"
 #include "Chaos/Serialization/SolverSerializer.h"
+#include "Containers/AnsiString.h"
+#include "Chaos/PhysicsObject.h"
 
 FPhysicsDelegatesCore::FOnUpdatePhysXMaterial FPhysicsDelegatesCore::OnUpdatePhysXMaterial;
 
@@ -1392,6 +1394,14 @@ FPhysicsConstraintHandle FChaosEngineInterface::CreateConstraint(Chaos::FPhysics
 				Params.Scene = Scene;
 				Params.bStatic = false;
 				Params.InitialTM = FTransform::Identity;
+
+#if CHAOS_DEBUG_NAME
+				TAnsiStringBuilder<256> EndPointDebugNameBuilder;
+				EndPointDebugNameBuilder.Appendf("Constraint Kinematic Endpoint - [%s]", TCHAR_TO_ANSI(*Chaos::FPhysicsObjectInterface::GetDebugName(ValidObject)));
+
+				Params.DebugName = EndPointDebugNameBuilder.GetData();
+#endif
+
 				FChaosEngineInterface::CreateActor(Params, KinematicEndPoint);
 
 				// Chaos requires our particles have geometry.
