[Chaos] Change to allow AddForce related methods to generate wake events.

There was an issue with the current implementation where we would manually set the object state in the Game thread API without allowing wake/sleep events to be dispatched. but then when we process the update from the Physics Thread, we wouldn't dispatch the wake event because the GT particle state was already set to dynamic

diff --git a/Engine/Source/Runtime/PhysicsCore/Private/ChaosEngineInterface.cpp b/Engine/Source/Runtime/PhysicsCore/Private/ChaosEngineInterface.cpp
--- a/Engine/Source/Runtime/PhysicsCore/Private/ChaosEngineInterface.cpp
+++ b/Engine/Source/Runtime/PhysicsCore/Private/ChaosEngineInterface.cpp
@@ -908,12 +908,11 @@ struct FChaosStateOps
 	}
 
 	static void AddForce(BodyHandleType& BodyHandle, const FVector& Force, bool bAllowSubstepping, bool bAccelChange)
-	{
-		Chaos::EObjectStateType ObjectState = BodyHandle.ObjectState();
-		
+	{		
 		if (!Force.IsZero())
 		{
-			BodyHandle.SetObjectState(Chaos::EObjectStateType::Dynamic, true);
+			constexpr bool bAllowEvents = true;
+			BodyHandle.SetObjectState(Chaos::EObjectStateType::Dynamic, bAllowEvents);
 		}
 
 		if (bAccelChange)
@@ -930,11 +929,10 @@ struct FChaosStateOps
 
 	static void AddForceAtPosition(BodyHandleType& BodyHandle, const FVector& Force, const FVector& Position, bool bAllowSubstepping, bool bIsLocalForce /*= false*/)
 	{
-		
-		Chaos::EObjectStateType ObjectState = BodyHandle.ObjectState();
 		const Chaos::FVec3 WorldCOM = Chaos::FParticleUtilitiesGT::GetCoMWorldPosition(&BodyHandle);
 
-		BodyHandle.SetObjectState(Chaos::EObjectStateType::Dynamic);
+		constexpr bool bAllowEvents = true;
+		BodyHandle.SetObjectState(Chaos::EObjectStateType::Dynamic, bAllowEvents);
 
 		if (bIsLocalForce)
 		{
@@ -967,7 +965,8 @@ struct FChaosStateOps
 				return;
 			}
 
-			BodyHandle.SetObjectState(Chaos::EObjectStateType::Dynamic);
+			constexpr bool bAllowEvents = true;
+			BodyHandle.SetObjectState(Chaos::EObjectStateType::Dynamic, bAllowEvents);
 
 			if (Distance < 1e-4)
 			{
