Add rigid_body native verse class, which owns a chaos particle internally and provides a verse api for its manipulation.

#okforversepublic

diff --git a/Engine/Source/Runtime/PhysicsCore/Private/ChaosEngineInterface.cpp b/Engine/Source/Runtime/PhysicsCore/Private/ChaosEngineInterface.cpp
--- a/Engine/Source/Runtime/PhysicsCore/Private/ChaosEngineInterface.cpp
+++ b/Engine/Source/Runtime/PhysicsCore/Private/ChaosEngineInterface.cpp
@@ -759,9 +759,30 @@ FTransform FChaosEngineInterface::GetComTransformLocal_AssumesLocked(const FPhys
 	return FTransform();
 }
 
+void FChaosEngineInterface::SetComTransformLocal_AssumesLocked(const FPhysicsActorHandle& InActorReference, const FTransform& InTransform)
+{
+	if(ensure(FChaosEngineInterface::IsValid(InActorReference)))
+	{
+		InActorReference->GetGameThreadAPI().SetRotationOfMass(InTransform.GetRotation());
+		InActorReference->GetGameThreadAPI().SetCenterOfMass(InTransform.GetTranslation());
+	}
+}
+
 FVector FChaosEngineInterface::GetLocalInertiaTensor_AssumesLocked(const FPhysicsActorHandle& InActorReference)
 {
-	return FVector(InActorReference->GetGameThreadAPI().I());
+	if(ensure(FChaosEngineInterface::IsValid(InActorReference)))
+	{
+		return FVector(InActorReference->GetGameThreadAPI().I());
+	}
+	return FVector::ZeroVector;
+}
+
+void FChaosEngineInterface::SetLocalInertiaTensor_AssumesLocked(const FPhysicsActorHandle& InActorReference, const FVector& InInertiaTensor)
+{
+	if(ensure(FChaosEngineInterface::IsValid(InActorReference)))
+	{
+		InActorReference->GetGameThreadAPI().SetI(InInertiaTensor);
+	}
 }
 
 FBox FChaosEngineInterface::GetBounds_AssumesLocked(const FPhysicsActorHandle& InActorReference)
diff --git a/Engine/Source/Runtime/PhysicsCore/Public/Chaos/ChaosEngineInterface.h b/Engine/Source/Runtime/PhysicsCore/Public/Chaos/ChaosEngineInterface.h
--- a/Engine/Source/Runtime/PhysicsCore/Public/Chaos/ChaosEngineInterface.h
+++ b/Engine/Source/Runtime/PhysicsCore/Public/Chaos/ChaosEngineInterface.h
@@ -445,8 +445,10 @@ public:
 
 	static PHYSICSCORE_API FTransform GetComTransform_AssumesLocked(const FPhysicsActorHandle& InActorReference);
 	static PHYSICSCORE_API FTransform GetComTransformLocal_AssumesLocked(const FPhysicsActorHandle& InActorReference);
+	static PHYSICSCORE_API void SetComTransformLocal_AssumesLocked(const FPhysicsActorHandle& InActorReference, const FTransform& InTransform);
 
 	static PHYSICSCORE_API FVector GetLocalInertiaTensor_AssumesLocked(const FPhysicsActorHandle& InActorReference);
+	static PHYSICSCORE_API void SetLocalInertiaTensor_AssumesLocked(const FPhysicsActorHandle& InActorReference, const FVector& InInertiaTensor);
 	static PHYSICSCORE_API FBox GetBounds_AssumesLocked(const FPhysicsActorHandle& InActorReference);
 	static PHYSICSCORE_API FBox GetBounds_AssumesLocked(const FPhysicsActorHandle& InActorReference, const FTransform& InTransform);
 	static PHYSICSCORE_API FBox GetBoundsLocal_AssumesLocked(const FPhysicsActorHandle& InActorReference);
