Added a safeguard for sweeps and overlaps having a nan object.

diff --git a/Engine/Source/Runtime/PhysicsCore/Public/SQVisitor.h b/Engine/Source/Runtime/PhysicsCore/Public/SQVisitor.h
--- a/Engine/Source/Runtime/PhysicsCore/Public/SQVisitor.h
+++ b/Engine/Source/Runtime/PhysicsCore/Public/SQVisitor.h
@@ -638,17 +638,22 @@ void SweepHelper(const QueryGeomType& QueryGeom, const Chaos::ISpatialAccelerati
 	using namespace Chaos;
 	using namespace ChaosInterface;
 
-	const FAABB3 Bounds = QueryGeom.BoundingBox().TransformedAABB(StartTM);
-	TSQVisitor<QueryGeomType, TPayload, TSweepHit, std::is_same<TSweepHit, FSweepHit>::value> SweepVisitor(StartTM, Dir, HitBuffer, OutputFlags, QueryFilterData, QueryCallback, QueryGeom, DebugParams);
+	const FAABB3 Bounds = QueryGeom.CalculateTransformedBounds(StartTM);
+	// There are cases where an invalid shape comes in which produces an invalid bounds. In this case, nan will enter causing the entire tree to be queried which will timeout.
+	const bool bIsValid = !Bounds.Min().ContainsNaN() && !Bounds.Max().ContainsNaN();
+	if (ensure(bIsValid))
+	{
+		TSQVisitor<QueryGeomType, TPayload, TSweepHit, std::is_same<TSweepHit, FSweepHit>::value> SweepVisitor(StartTM, Dir, HitBuffer, OutputFlags, QueryFilterData, QueryCallback, QueryGeom, DebugParams);
 
-	HitBuffer.IncFlushCount();
+		HitBuffer.IncFlushCount();
 
-	// NOTE: Previously, we'd attempt to perform a sweep as an overlap if the DeltaMagnitude was 0. 
-	// This caused various discrepancies as some of the out info couldn't be computed. Instead, sweep has been updated to handle 0 length sweeps.
-	const FVector HalfExtents = Bounds.Extents() * 0.5f;
-	SpatialAcceleration.Sweep(Bounds.GetCenter(), Dir, DeltaMagnitude, HalfExtents, SweepVisitor);
+		// NOTE: Previously, we'd attempt to perform a sweep as an overlap if the DeltaMagnitude was 0. 
+		// This caused various discrepancies as some of the out info couldn't be computed. Instead, sweep has been updated to handle 0 length sweeps.
+		const FVector HalfExtents = Bounds.Extents() * 0.5f;
+		SpatialAcceleration.Sweep(Bounds.GetCenter(), Dir, DeltaMagnitude, HalfExtents, SweepVisitor);
 
-	HitBuffer.DecFlushCount();
+		HitBuffer.DecFlushCount();
+	}
 }
 
 
@@ -659,21 +664,25 @@ void OverlapHelper(const QueryGeomType& QueryGeom, const Chaos::ISpatialAccelera
 	using namespace ChaosInterface;
 
 	const FAABB3 Bounds = QueryGeom.CalculateTransformedBounds(GeomPose);
+	// There are cases where an invalid shape comes in which produces an invalid bounds. In this case, nan will enter causing the entire tree to be queried which will timeout.
+	const bool bIsValid = !Bounds.Min().ContainsNaN() && !Bounds.Max().ContainsNaN();
+	if (ensure(bIsValid))
+	{
+		HitBuffer.IncFlushCount();
 
-	HitBuffer.IncFlushCount();
-
-	bool bSkipNarrowPhase = QueryFilterData.flags & FPhysicsQueryFlag::eSKIPNARROWPHASE;
+		bool bSkipNarrowPhase = QueryFilterData.flags & FPhysicsQueryFlag::eSKIPNARROWPHASE;
 
-	constexpr bool bGTData = std::is_same<TOverlapHit, FOverlapHit>::value;
-	if (bSkipNarrowPhase)
-	{
-		TBPVisitor<QueryGeomType, TPayload, TOverlapHit, bGTData> OverlapVisitor(GeomPose, HitBuffer, QueryFilterData, QueryCallback, QueryGeom, DebugParams);
-		SpatialAcceleration.Overlap(Bounds, OverlapVisitor);
-	}
-	else
-	{
-		TSQVisitor<QueryGeomType, TPayload, TOverlapHit, bGTData> OverlapVisitor(GeomPose, HitBuffer, QueryFilterData, QueryCallback, QueryGeom, DebugParams);
-		SpatialAcceleration.Overlap(Bounds, OverlapVisitor);
+		constexpr bool bGTData = std::is_same<TOverlapHit, FOverlapHit>::value;
+		if (bSkipNarrowPhase)
+		{
+			TBPVisitor<QueryGeomType, TPayload, TOverlapHit, bGTData> OverlapVisitor(GeomPose, HitBuffer, QueryFilterData, QueryCallback, QueryGeom, DebugParams);
+			SpatialAcceleration.Overlap(Bounds, OverlapVisitor);
+		}
+		else
+		{
+			TSQVisitor<QueryGeomType, TPayload, TOverlapHit, bGTData> OverlapVisitor(GeomPose, HitBuffer, QueryFilterData, QueryCallback, QueryGeom, DebugParams);
+			SpatialAcceleration.Overlap(Bounds, OverlapVisitor);
+		}
+		HitBuffer.DecFlushCount();
 	}
-	HitBuffer.DecFlushCount();
 }
