Collision Filter Refactor Part 2:
Deprecated the old API methods that used FCollisionFilterData and added new methods using the new types.
Did minimal cleanup of old call sites to to use the new API methods.

diff --git a/Engine/Plugins/Experimental/ChaosMover/Source/ChaosMover/Private/ChaosMoverSimulation.cpp b/Engine/Plugins/Experimental/ChaosMover/Source/ChaosMover/Private/ChaosMoverSimulation.cpp
--- a/Engine/Plugins/Experimental/ChaosMover/Source/ChaosMover/Private/ChaosMoverSimulation.cpp
+++ b/Engine/Plugins/Experimental/ChaosMover/Source/ChaosMover/Private/ChaosMoverSimulation.cpp
@@ -248,21 +248,24 @@ void UChaosMoverSimulation::ModifyContacts(const FMoverTimeStep& TimeStep, const
 
 					if (const Chaos::FShapeInstance* Shape = PairModifier.GetShape(OtherIdx))
 					{
-						uint32 ComponentID = ChaosInterface::GetSimulationFilterData(*Shape).Word2;
+						const Chaos::Filter::FCombinedShapeFilterData& CombinedShapeFilter = ChaosInterface::GetCombinedShapeFilterData(*Shape);
+						const Chaos::Filter::FInstanceData& InstanceData = CombinedShapeFilter.GetInstanceData();
+						const Chaos::Filter::FShapeFilterData& FilterData = CombinedShapeFilter.GetShapeFilterData();
+						uint32 ComponentID = InstanceData.GetComponentId();
 						if (SimInputs->CollisionQueryParams.GetIgnoredComponents().Contains(ComponentID))
 						{
 							PairModifier.Disable();
 							continue;
 						}
 
-						uint32 ActorID = ChaosInterface::GetQueryFilterData(*Shape).Word0;
+						uint32 ActorID = InstanceData.GetActorId();
 						if (SimInputs->CollisionQueryParams.GetIgnoredSourceObjects().Contains(ActorID))
 						{
 							PairModifier.Disable();
 							continue;
 						}
 
-						FMaskFilter ShapeMaskFilter = ChaosInterface::GetQueryFilterData(*Shape).Word3 >> (32u - NumExtraFilterBits);
+						FMaskFilter ShapeMaskFilter = FilterData.GetMaskFilter();
 						if (SimInputs->CollisionQueryParams.IgnoreMask & ShapeMaskFilter)
 						{
 							PairModifier.Disable();
diff --git a/Engine/Plugins/Experimental/SpatialReadiness/Source/SpatialReadiness/Private/SpatialReadiness/SpatialReadinessSimCallback.cpp b/Engine/Plugins/Experimental/SpatialReadiness/Source/SpatialReadiness/Private/SpatialReadiness/SpatialReadinessSimCallback.cpp
--- a/Engine/Plugins/Experimental/SpatialReadiness/Source/SpatialReadiness/Private/SpatialReadiness/SpatialReadinessSimCallback.cpp
+++ b/Engine/Plugins/Experimental/SpatialReadiness/Source/SpatialReadiness/Private/SpatialReadiness/SpatialReadinessSimCallback.cpp
@@ -223,7 +223,7 @@ bool FSpatialReadinessSimCallback::QueryReadiness_GT(const FBox& Bounds, TArray<
 	QueryParams.bTraceComplex = false;
 	FCollisionObjectQueryParams ObjectParams(static_cast<uint8>(Channel));
 	FCollisionResponseContainer ResponseContainer;
-	FCollisionFilterData CollisionFilterData = CreateQueryFilterData(ChannelBit, bComplex, ResponseContainer, QueryParams, ObjectParams, bMulti);
+	FCollisionFilterData CollisionFilterData = Chaos::Filter::FQueryFilterBuilder::GetLegacyQueryFilter(CreateChaosQueryFilterData(ChannelBit, bComplex, ResponseContainer, QueryParams, ObjectParams, bMulti));
 	FOverlapAllQueryCallback QueryCallback;
 	EQueryFlags QueryFlags = EQueryFlags::PreFilter;// | EQueryFlags::SkipNarrowPhase;
 	ChaosInterface::FQueryFilterData QueryFilterData = ChaosInterface::MakeQueryFilterData(CollisionFilterData, QueryFlags, QueryParams);
diff --git a/Engine/Source/Runtime/Engine/Private/Collision/CollisionConversions.cpp b/Engine/Source/Runtime/Engine/Private/Collision/CollisionConversions.cpp
--- a/Engine/Source/Runtime/Engine/Private/Collision/CollisionConversions.cpp
+++ b/Engine/Source/Runtime/Engine/Private/Collision/CollisionConversions.cpp
@@ -349,7 +349,7 @@ EConvertQueryResult ConvertQueryImpactHitImp(const UWorld* World, const THitLoca
 	const FPhysicsShape* GTShape = GetGTShape<bIsGTQuery>(pHitShape, GTActor);
 
 	// See if this is a 'blocking' hit
-	const FCollisionFilterData ShapeFilter = GetQueryFilterData(HitShape);
+	const FCollisionFilterData ShapeFilter = Chaos::Filter::FShapeFilterBuilder::GetLegacyShapeQueryFilter(GetCombinedShapeFilterData(HitShape));
 	const ECollisionQueryHitType HitType = FCollisionQueryFilterCallback::CalcQueryHitType(QueryFilter, ShapeFilter);
 	OutResult.bBlockingHit = (HitType == ECollisionQueryHitType::Block);
 	OutResult.bStartPenetrating = bInitialOverlap;
@@ -563,7 +563,7 @@ static bool ConvertOverlappedShapeToImpactHit(const UWorld* World, const THitLoc
 	const auto& HitActor = *pHitActor;
 
 	// See if this is a 'blocking' hit
-	FCollisionFilterData ShapeFilter = GetQueryFilterData(HitShape);
+	FCollisionFilterData ShapeFilter = Chaos::Filter::FShapeFilterBuilder::GetLegacyShapeQueryFilter(GetCombinedShapeFilterData(HitShape));
 	const ECollisionQueryHitType HitType = FCollisionQueryFilterCallback::CalcQueryHitType(QueryFilter, ShapeFilter);
 	const bool bBlockingHit = (HitType == ECollisionQueryHitType::Block);
 	OutResult.bBlockingHit = bBlockingHit;
@@ -782,7 +782,7 @@ bool IsBlocking(const FPhysicsShape& Shape, const FCollisionFilterData& QueryFil
 	using namespace ChaosInterface;
 
 	// See if this is a 'blocking' hit
-	const FCollisionFilterData ShapeFilter = GetQueryFilterData(Shape);
+	const FCollisionFilterData ShapeFilter = Chaos::Filter::FShapeFilterBuilder::GetLegacyShapeQueryFilter(GetCombinedShapeFilterData(Shape));
 	const ECollisionQueryHitType HitType = FCollisionQueryFilterCallback::CalcQueryHitType(QueryFilter, ShapeFilter);
 	const bool bBlock = (HitType == ECollisionQueryHitType::Block);
 	return bBlock;
diff --git a/Engine/Source/Runtime/Engine/Private/Collision/SceneQuery.cpp b/Engine/Source/Runtime/Engine/Private/Collision/SceneQuery.cpp
--- a/Engine/Source/Runtime/Engine/Private/Collision/SceneQuery.cpp
+++ b/Engine/Source/Runtime/Engine/Private/Collision/SceneQuery.cpp
@@ -527,7 +527,8 @@ bool TSceneCastCommonImpWithRetryRequest(const UWorld* World, typename Traits::O
 	{
 		// Create filter data used to filter collisions 
 		CA_SUPPRESS(6326);
-		FCollisionFilterData Filter = CreateQueryFilterData(TraceChannel, Params.bTraceComplex, ResponseParams.CollisionResponse, Params, ObjectParams, Traits::SingleMultiOrTest == ESingleMultiOrTest::Multi);
+		const Chaos::Filter::FQueryFilterData QueryFilterData = CreateChaosQueryFilterData(TraceChannel, Params.bTraceComplex, ResponseParams.CollisionResponse, Params, ObjectParams, Traits::SingleMultiOrTest == ESingleMultiOrTest::Multi);
+		FCollisionFilterData Filter = Chaos::Filter::FQueryFilterBuilder::GetLegacyQueryFilter(QueryFilterData);
 		
 		CA_SUPPRESS(6326);
 		FCollisionQueryFilterCallback QueryCallback(Params, Traits::GeometryQuery == ESweepOrRay::Sweep);
@@ -1113,7 +1114,8 @@ bool GeomOverlapMultiImp(const UWorld* World, const FPhysicsGeometry& Geom, cons
 		CVD_TRACE_SCOPED_SCENE_QUERY_HELPER(World, &Geom, GeomPose, FVector::ZeroVector, TraceChannel, Params, ResponseParams, ObjectParams, EChaosVDSceneQueryType::Overlap, (InfoType == EQueryInfo::GatherAll) ? EChaosVDSceneQueryMode::Multi : EChaosVDSceneQueryMode::Test, bIsRetryQuery);
 
 		// Create filter data used to filter collisions
-		FCollisionFilterData Filter = CreateQueryFilterData(TraceChannel, Params.bTraceComplex, ResponseParams.CollisionResponse, Params, ObjectParams, InfoType != EQueryInfo::IsAnything);
+		const Chaos::Filter::FQueryFilterData QueryFilterData = CreateChaosQueryFilterData(TraceChannel, Params.bTraceComplex, ResponseParams.CollisionResponse, Params, ObjectParams, InfoType != EQueryInfo::IsAnything);
+		FCollisionFilterData Filter = Chaos::Filter::FQueryFilterBuilder::GetLegacyQueryFilter(QueryFilterData);
 		FCollisionQueryFilterCallback QueryCallback(Params, false);
 		QueryCallback.bIgnoreTouches |= (InfoType == EQueryInfo::IsBlocking); // pre-filter to ignore touches and only get blocking hits, if that's what we're after.
 		QueryCallback.bIsOverlapQuery = true;
diff --git a/Engine/Source/Runtime/Engine/Private/Collision/WorldCollision.cpp b/Engine/Source/Runtime/Engine/Private/Collision/WorldCollision.cpp
--- a/Engine/Source/Runtime/Engine/Private/Collision/WorldCollision.cpp
+++ b/Engine/Source/Runtime/Engine/Private/Collision/WorldCollision.cpp
@@ -525,8 +525,8 @@ bool UWorld::ComponentSweepMultiByChannel(TArray<struct FHitResult>& OutHits, cl
 		TArray<Chaos::FShapeInstanceProxy*> Shapes = Interface->GetAllThreadShapes({ &Object, 1 });
 		for (Chaos::FShapeInstanceProxy* Shape : Shapes)
 		{
-			FCollisionFilterData ShapeFilter = ChaosInterface::GetQueryFilterData(*Shape);
-			uint32 ShapeActorID = ShapeFilter.Word0; // Unique Actor ID is saved in word 0
+			const Chaos::Filter::FInstanceData InstanceData = ChaosInterface::GetFilterInstanceData(*Shape);
+			uint32 ShapeActorID = InstanceData.GetActorId();
 			ActorsToExclude.Add(ShapeActorID);
 		}
 
diff --git a/Engine/Source/Runtime/Engine/Private/PhysicsEngine/BodyInstance.cpp b/Engine/Source/Runtime/Engine/Private/PhysicsEngine/BodyInstance.cpp
--- a/Engine/Source/Runtime/Engine/Private/PhysicsEngine/BodyInstance.cpp
+++ b/Engine/Source/Runtime/Engine/Private/PhysicsEngine/BodyInstance.cpp
@@ -1105,8 +1105,9 @@ void FBodyInstance::UpdatePhysicsFilterData()
 				}
 
 				// Apply new collision settings to this shape
-				FPhysicsInterface::SetSimulationFilter(InnerShape, FilterData.SimFilter);
-				FPhysicsInterface::SetQueryFilter(InnerShape, bIsTrimesh ? FilterData.QueryComplexFilter : FilterData.QuerySimpleFilter);
+				const FCollisionFilterData& QueryFilter = bIsTrimesh ? FilterData.QueryComplexFilter : FilterData.QuerySimpleFilter;
+				Chaos::Filter::FCombinedShapeFilterData CombinedShapeFilter = Chaos::Filter::FShapeFilterBuilder::BuildFromLegacyShapeFilter(QueryFilter, FilterData.SimFilter);
+				FPhysicsInterface::SetCombinedShapeFilterData(InnerShape, CombinedShapeFilter);
 			});
 		}
 
diff --git a/Engine/Source/Runtime/Engine/Private/PhysicsEngine/CollisionQueryFilterCallback.cpp b/Engine/Source/Runtime/Engine/Private/PhysicsEngine/CollisionQueryFilterCallback.cpp
--- a/Engine/Source/Runtime/Engine/Private/PhysicsEngine/CollisionQueryFilterCallback.cpp
+++ b/Engine/Source/Runtime/Engine/Private/PhysicsEngine/CollisionQueryFilterCallback.cpp
@@ -89,13 +89,15 @@ ECollisionQueryHitType FCollisionQueryFilterCallback::PreFilterBaseImp(const FCo
 		return ECollisionQueryHitType::None;
 	}
 
-	FCollisionFilterData ShapeFilter = ChaosInterface::GetQueryFilterData(Shape);
+	const Chaos::Filter::FCombinedShapeFilterData CombinedShapeFilter = ChaosInterface::GetCombinedShapeFilterData(Shape);
+	const Chaos::Filter::FInstanceData& InstanceData = CombinedShapeFilter.GetInstanceData();
+	FCollisionFilterData ShapeFilter = Chaos::Filter::FShapeFilterBuilder::GetLegacyShapeQueryFilter(CombinedShapeFilter);
 
 	// We usually don't have ignore components so we try to avoid the virtual getSimulationFilterData() call below. 'word2' of shape sim filter data is componentID.
 	uint32 ComponentID = 0;
 	if (IgnoreComponents.Num() > 0)
 	{
-		ComponentID = ChaosInterface::GetSimulationFilterData(Shape).Word2;
+		ComponentID = InstanceData.GetComponentId();
 	}
 
 	FBodyInstance* BodyInstance = nullptr;
diff --git a/Engine/Source/Runtime/Engine/Private/PhysicsEngine/PhysicsInterfaceUtils.cpp b/Engine/Source/Runtime/Engine/Private/PhysicsEngine/PhysicsInterfaceUtils.cpp
--- a/Engine/Source/Runtime/Engine/Private/PhysicsEngine/PhysicsInterfaceUtils.cpp
+++ b/Engine/Source/Runtime/Engine/Private/PhysicsEngine/PhysicsInterfaceUtils.cpp
@@ -103,3 +103,12 @@ FCollisionFilterData CreateQueryFilterData(const uint8 MyChannel, const bool bTr
 		return CreateTraceQueryFilterData(MyChannel, bTraceComplex, InCollisionResponseContainer, QueryParam);
 	}
 }
+
+Chaos::Filter::FQueryFilterData CreateChaosQueryFilterData(const uint8 MyChannel, const bool bTraceComplex, const FCollisionResponseContainer& InCollisionResponseContainer, const struct FCollisionQueryParams& QueryParam, const struct FCollisionObjectQueryParams& ObjectParam, const bool bMultitrace)
+{
+	// TODO @ JoshD: Update
+	PRAGMA_DISABLE_DEPRECATION_WARNINGS
+	FCollisionFilterData FilterData = CreateQueryFilterData(MyChannel, bTraceComplex, InCollisionResponseContainer, QueryParam, ObjectParam, bMultitrace);
+	return Chaos::Filter::FQueryFilterBuilder::BuildFromLegacyQueryFilter(FilterData);
+	PRAGMA_ENABLE_DEPRECATION_WARNINGS
+}
diff --git a/Engine/Source/Runtime/Engine/Public/Physics/Experimental/ChaosInterfaceWrapper.h b/Engine/Source/Runtime/Engine/Public/Physics/Experimental/ChaosInterfaceWrapper.h
--- a/Engine/Source/Runtime/Engine/Public/Physics/Experimental/ChaosInterfaceWrapper.h
+++ b/Engine/Source/Runtime/Engine/Public/Physics/Experimental/ChaosInterfaceWrapper.h
@@ -36,7 +36,9 @@ struct FScopedSceneReadLock
 
 inline FQueryFilterData MakeQueryFilterData(const FCollisionFilterData& FilterData, EQueryFlags QueryFlags, const FCollisionQueryParams& Params)
 {
+	PRAGMA_DISABLE_INTERNAL_WARNINGS
 	return FChaosQueryFilterData(U2CFilterData(FilterData), U2CQueryFlags(QueryFlags));
+	PRAGMA_ENABLE_INTERNAL_WARNINGS
 }
 
 FBodyInstance* GetUserData(const Chaos::FGeometryParticle& Actor);
diff --git a/Engine/Source/Runtime/Engine/Public/Physics/PhysicsInterfaceUtils.h b/Engine/Source/Runtime/Engine/Public/Physics/PhysicsInterfaceUtils.h
--- a/Engine/Source/Runtime/Engine/Public/Physics/PhysicsInterfaceUtils.h
+++ b/Engine/Source/Runtime/Engine/Public/Physics/PhysicsInterfaceUtils.h
@@ -16,7 +16,8 @@ struct FConstraintInstance;
 // FILTER DATA
 
 /** Utility for creating a filter data object for performing a query (trace) against the scene */
-ENGINE_API FCollisionFilterData CreateQueryFilterData(const uint8 MyChannel, const bool bTraceComplex, const FCollisionResponseContainer& InCollisionResponseContainer, const struct FCollisionQueryParams& QueryParam, const struct FCollisionObjectQueryParams & ObjectParam, const bool bMultitrace);
+UE_DEPRECATED(5.7, "Use CreateChaosQueryFilterData instead") ENGINE_API FCollisionFilterData CreateQueryFilterData(const uint8 MyChannel, const bool bTraceComplex, const FCollisionResponseContainer& InCollisionResponseContainer, const struct FCollisionQueryParams& QueryParam, const struct FCollisionObjectQueryParams & ObjectParam, const bool bMultitrace);
+ENGINE_API Chaos::Filter::FQueryFilterData CreateChaosQueryFilterData(const uint8 MyChannel, const bool bTraceComplex, const FCollisionResponseContainer& InCollisionResponseContainer, const struct FCollisionQueryParams& QueryParam, const struct FCollisionObjectQueryParams & ObjectParam, const bool bMultitrace);
 
 struct FConstraintBrokenDelegateData
 {
diff --git a/Engine/Source/Runtime/Experimental/Chaos/Private/Chaos/PhysicsObjectInterface.cpp b/Engine/Source/Runtime/Experimental/Chaos/Private/Chaos/PhysicsObjectInterface.cpp
--- a/Engine/Source/Runtime/Experimental/Chaos/Private/Chaos/PhysicsObjectInterface.cpp
+++ b/Engine/Source/Runtime/Experimental/Chaos/Private/Chaos/PhysicsObjectInterface.cpp
@@ -1331,6 +1331,13 @@ namespace Chaos
 
 	template<EThreadContext Id>
 	void FWritePhysicsObjectInterface<Id>::UpdateShapeFilterData(TArrayView<const FPhysicsObjectHandle> InObjects, const FCollisionFilterData& QueryData, const FCollisionFilterData& SimData)
+	{
+		Filter::FCombinedShapeFilterData ShapeFilterData = Chaos::Filter::FShapeFilterBuilder::BuildFromLegacyShapeFilter(QueryData, SimData);
+		UpdateCombinedShapeFilterData(InObjects, ShapeFilterData);
+	}
+
+	template<EThreadContext Id>
+	void FWritePhysicsObjectInterface<Id>::UpdateCombinedShapeFilterData(TArrayView<const FPhysicsObjectHandle> InObjects, const Chaos::Filter::FCombinedShapeFilterData& CombinedShapeFilter)
 	{
 		for (const FPhysicsObjectHandle Object : InObjects)
 		{
@@ -1343,8 +1350,7 @@ namespace Chaos
 			{
 				for (const TUniquePtr<Chaos::FPerShapeData>& ShapeData : Particle->ShapesArray())
 				{
-					ShapeData->SetQueryData(QueryData);
-					ShapeData->SetSimData(SimData);
+					ShapeData->SetCombinedShapeFilterData(CombinedShapeFilter);
 				}
 			}
 		}
diff --git a/Engine/Source/Runtime/Experimental/Chaos/Private/Chaos/ShapeInstance.cpp b/Engine/Source/Runtime/Experimental/Chaos/Private/Chaos/ShapeInstance.cpp
--- a/Engine/Source/Runtime/Experimental/Chaos/Private/Chaos/ShapeInstance.cpp
+++ b/Engine/Source/Runtime/Experimental/Chaos/Private/Chaos/ShapeInstance.cpp
@@ -193,7 +193,35 @@ namespace Chaos
 		check(false);
 	}
 
+	const Filter::FShapeFilterData FPerShapeData::GetShapeFilterData() const
+	{
+		return DownCast([](const auto& ShapeInstance) -> const auto { return ShapeInstance.GetShapeFilterData(); });
+	}
+
+	void FPerShapeData::SetShapeFilterData(const Filter::FShapeFilterData& InShapeFilter)
+	{
+		return DownCast([&InShapeFilter](auto& ShapeInstance) { ShapeInstance.SetShapeFilterData(InShapeFilter); });
+	}
+
+	const Filter::FInstanceData FPerShapeData::GetFilterInstanceData() const
+	{
+		return DownCast([](const auto& ShapeInstance) -> const auto { return ShapeInstance.GetFilterInstanceData(); });
+	}
+
+	void FPerShapeData::SetFilterInstanceData(const Filter::FInstanceData& InInstanceData)
+	{
+		return DownCast([&InInstanceData](auto& ShapeInstance) { ShapeInstance.SetFilterInstanceData(InInstanceData); });
+	}
+
+	const Filter::FCombinedShapeFilterData FPerShapeData::GetCombinedShapeFilterData() const
+	{
+		return DownCast([](const auto& ShapeInstance) -> const auto { return ShapeInstance.GetCombinedShapeFilterData(); });
+	}
 
+	void FPerShapeData::SetCombinedShapeFilterData(const Filter::FCombinedShapeFilterData& InShapeCombinedFilter)
+	{
+		return DownCast([&InShapeCombinedFilter](auto& ShapeInstance) { ShapeInstance.SetCombinedShapeFilterData(InShapeCombinedFilter); });
+	}
 
 	TUniquePtr<FShapeInstanceProxy> FShapeInstanceProxy::Make(int32 InShapeIdx, const FImplicitObjectPtr& InGeometry)
 	{
@@ -205,6 +233,35 @@ namespace Chaos
 		InOutShapePtr->Geometry = InGeometry;
 	}
 
+	const Filter::FShapeFilterData FShapeInstanceProxy::GetShapeFilterData() const
+	{
+		return CollisionData.Read().GetShapeFilterData();
+	}
+
+	void FShapeInstanceProxy::SetShapeFilterData(const Filter::FShapeFilterData& InShapeFilter)
+	{
+		CollisionData.Modify(true, DirtyFlags, Proxy, ShapeIdx, [InShapeFilter](FCollisionData& Data) { Data.SetShapeFilterData(InShapeFilter); });
+	}
+
+	const Filter::FInstanceData FShapeInstanceProxy::GetFilterInstanceData() const
+	{
+		return CollisionData.Read().GetFilterInstanceData();
+	}
+
+	void FShapeInstanceProxy::SetFilterInstanceData(const Filter::FInstanceData& InInstanceData)
+	{
+		CollisionData.Modify(true, DirtyFlags, Proxy, ShapeIdx, [InInstanceData](FCollisionData& Data) { Data.SetFilterInstanceData(InInstanceData); });
+	}
+
+	const Filter::FCombinedShapeFilterData FShapeInstanceProxy::GetCombinedShapeFilterData() const
+	{
+		return CollisionData.Read().GetCombinedShapeFilterData();
+	}
+
+	void FShapeInstanceProxy::SetCombinedShapeFilterData(const Filter::FCombinedShapeFilterData& InShapeCombinedFilter)
+	{
+		CollisionData.Modify(true, DirtyFlags, Proxy, ShapeIdx, [InShapeCombinedFilter](FCollisionData& Data) { Data.SetCombinedShapeFilterData(InShapeCombinedFilter); });
+	}
 
 
 	TUniquePtr<FShapeInstance> FShapeInstance::Make(int32 InShapeIdx, const FImplicitObjectPtr& InGeometry)
@@ -324,6 +381,36 @@ namespace Chaos
 		}
 	}
 
+	const Filter::FShapeFilterData FShapeInstance::GetShapeFilterData() const
+	{
+		return CollisionData.GetShapeFilterData();
+	}
+
+	void FShapeInstance::SetShapeFilterData(const Filter::FShapeFilterData& InShapeFilter)
+	{
+		CollisionData.SetShapeFilterData(InShapeFilter);
+	}
+
+	const Filter::FInstanceData FShapeInstance::GetFilterInstanceData() const
+	{
+		return CollisionData.GetFilterInstanceData();
+	}
+
+	void FShapeInstance::SetFilterInstanceData(const Filter::FInstanceData& InInstanceData)
+	{
+		CollisionData.SetFilterInstanceData(InInstanceData);
+	}
+
+	const Filter::FCombinedShapeFilterData FShapeInstance::GetCombinedShapeFilterData() const
+	{
+		return CollisionData.GetCombinedShapeFilterData();
+	}
+
+	void FShapeInstance::SetCombinedShapeFilterData(const Filter::FCombinedShapeFilterData& InCombinedShapeFilter)
+	{
+		CollisionData.SetCombinedShapeFilterData(InCombinedShapeFilter);
+	}
+
 	///////////////////////////////////////////////////////////////////////////////////////////////
 	///////////////////////////////////////////////////////////////////////////////////////////////
 	// 
diff --git a/Engine/Source/Runtime/Experimental/Chaos/Private/PhysicsProxy/GeometryCollectionPhysicsProxy.cpp b/Engine/Source/Runtime/Experimental/Chaos/Private/PhysicsProxy/GeometryCollectionPhysicsProxy.cpp
--- a/Engine/Source/Runtime/Experimental/Chaos/Private/PhysicsProxy/GeometryCollectionPhysicsProxy.cpp
+++ b/Engine/Source/Runtime/Experimental/Chaos/Private/PhysicsProxy/GeometryCollectionPhysicsProxy.cpp
@@ -4638,11 +4638,18 @@ bool FGeometryCollectionPhysicsProxy::PullFromPhysicsState(const Chaos::FDirtyGe
 }
 
 void FGeometryCollectionPhysicsProxy::UpdateFilterData_External(const FCollisionFilterData& NewSimFilter, const FCollisionFilterData& NewQueryFilter)
+{
+	// SimFilter/QueryFilter members are read on both threads, these are const after initialization and are not updated here.
+	Chaos::Filter::FCombinedShapeFilterData NewShapeFilter = Chaos::Filter::FShapeFilterBuilder::BuildFromLegacyShapeFilter(NewQueryFilter, NewSimFilter);
+	UpdateFilterData_External(NewShapeFilter);
+}
+
+void FGeometryCollectionPhysicsProxy::UpdateFilterData_External(const Chaos::Filter::FCombinedShapeFilterData& NewShapeFilter)
 {
 	SCOPE_CYCLE_COUNTER(STAT_GCUpdateFilterData);
 	check(IsInGameThread());
 
-	// SimFilter/QueryFilter members are read on both threads, these are const after initialization and are not updated here.
+	// ShapeFilter members are read on both threads, these are const after initialization and are not updated here.
 	for (int32 Index = 0; Index < NumEffectiveParticles; ++Index)
 	{
 		FParticle* P = GTParticles[Index].Get();
@@ -4653,24 +4660,28 @@ void FGeometryCollectionPhysicsProxy::UpdateFilterData_External(const FCollision
 			for (int32 ShapeIndex = 0; ShapeIndex < NumShapes; ++ShapeIndex)
 			{
 				Chaos::FPerShapeData* Shape = Shapes[ShapeIndex].Get();
-				Shape->SetSimData(NewSimFilter);
-				Shape->SetQueryData(NewQueryFilter);
+				Shape->SetCombinedShapeFilterData(NewShapeFilter);
 			}
 		}
 	}
 
-	ExecuteOnPhysicsThread(*this, [this, NewSimFilter, NewQueryFilter]()
+	ExecuteOnPhysicsThread(*this, [this, NewShapeFilter]()
 		{
-			SetFilterData_Internal(NewSimFilter, NewQueryFilter);
+			SetFilterData_Internal(NewShapeFilter);
 		});
 }
 
 void FGeometryCollectionPhysicsProxy::SetFilterData_Internal(const FCollisionFilterData& NewSimFilter, const FCollisionFilterData& NewQueryFilter)
+{
+	Chaos::Filter::FCombinedShapeFilterData NewShapeFilter = Chaos::Filter::FShapeFilterBuilder::BuildFromLegacyShapeFilter(NewQueryFilter, NewSimFilter);
+	SetFilterData_Internal(NewShapeFilter);
+}
+
+void FGeometryCollectionPhysicsProxy::SetFilterData_Internal(const Chaos::Filter::FCombinedShapeFilterData& NewShapeFilter)
 {
 	using namespace Chaos;
 
-	Parameters.SimulationFilterData = NewSimFilter;
-	Parameters.QueryFilterData = NewQueryFilter;
+	Chaos::Filter::FShapeFilterBuilder::GetLegacyShapeFilter(NewShapeFilter, Parameters.QueryFilterData, Parameters.SimulationFilterData);
 
 	TArray<FClusterHandle*> ProcessedInternalClusters;
 	if (Chaos::FPhysicsSolver* RigidSolver = GetSolver<Chaos::FPhysicsSolver>())
@@ -4685,8 +4696,7 @@ void FGeometryCollectionPhysicsProxy::SetFilterData_Internal(const FCollisionFil
 				const Chaos::FShapesArray& ShapesArray = Handle->ShapesArray();
 				for (const TUniquePtr<Chaos::FPerShapeData>& Shape : ShapesArray)
 				{
-					Shape->SetQueryData(Parameters.QueryFilterData);
-					Shape->SetSimData(Parameters.SimulationFilterData);
+					Shape->SetCombinedShapeFilterData(NewShapeFilter);
 				}
 
 				FClusterHandle* ParentHandle = Handle->Parent();
@@ -4738,7 +4748,7 @@ void FGeometryCollectionPhysicsProxy::SetPerParticleFilterData_Internal(const TA
 					check(Object != nullptr);
 					TArrayView<Chaos::FPhysicsObjectHandle> ParticleView{ &Object, 1 };
 					Interface.UpdateShapeCollisionFlags(ParticleView, Data.bSimEnabled, Data.bQueryEnabled);
-					Interface.UpdateShapeFilterData(ParticleView, Data.QueryFilter, Data.SimFilter);
+					Interface.UpdateCombinedShapeFilterData(ParticleView, Chaos::Filter::FShapeFilterBuilder::BuildFromLegacyShapeFilter(Data.QueryFilter, Data.SimFilter));
 					// todo(chaos): It's not ideal but the geometry collection needs to request the cluster union update its
 					// cached shape data if the particle is in a cluster union. This is because we don't share shape
 					// data between the GC shapes and the cluster union shapes.
diff --git a/Engine/Source/Runtime/Experimental/Chaos/Public/Chaos/PhysicsObjectInterface.h b/Engine/Source/Runtime/Experimental/Chaos/Public/Chaos/PhysicsObjectInterface.h
--- a/Engine/Source/Runtime/Experimental/Chaos/Public/Chaos/PhysicsObjectInterface.h
+++ b/Engine/Source/Runtime/Experimental/Chaos/Public/Chaos/PhysicsObjectInterface.h
@@ -133,7 +133,9 @@ namespace Chaos
 		CHAOS_API void SetAngularEtherDrag(TArrayView<const FPhysicsObjectHandle> InObjects, float InAngularDrag);
 
 		CHAOS_API void UpdateShapeCollisionFlags(TArrayView<const FPhysicsObjectHandle> InObjects, bool bSimCollision, bool bQueryCollision);
+		UE_DEPRECATED(5.7, "Use UpdateCombinedShapeFilterData instead")
 		CHAOS_API void UpdateShapeFilterData(TArrayView<const FPhysicsObjectHandle> InObjects, const FCollisionFilterData& QueryData, const FCollisionFilterData& SimData);
+		CHAOS_API void UpdateCombinedShapeFilterData(TArrayView<const FPhysicsObjectHandle> InObjects, const Chaos::Filter::FCombinedShapeFilterData& CombinedShapeFilter);
 
 		template<typename TPayloadType, typename T, int d>
 		void AddToSpatialAcceleration(TArrayView<const FPhysicsObjectHandle> InObjects, ISpatialAcceleration<TPayloadType, T, d>* SpatialAcceleration)
diff --git a/Engine/Source/Runtime/Experimental/Chaos/Public/Chaos/ShapeInstance.h b/Engine/Source/Runtime/Experimental/Chaos/Public/Chaos/ShapeInstance.h
--- a/Engine/Source/Runtime/Experimental/Chaos/Public/Chaos/ShapeInstance.h
+++ b/Engine/Source/Runtime/Experimental/Chaos/Public/Chaos/ShapeInstance.h
@@ -94,6 +94,13 @@ namespace Chaos
 		const FCollisionFilterData& GetSimData() const;
 		void SetSimData(const FCollisionFilterData& InSimData);
 
+		CHAOS_API const Filter::FShapeFilterData GetShapeFilterData() const;
+		CHAOS_API void SetShapeFilterData(const Filter::FShapeFilterData& InShapeFilter);
+		CHAOS_API const Filter::FInstanceData GetFilterInstanceData() const;
+		CHAOS_API void SetFilterInstanceData(const Filter::FInstanceData& InInstanceData);
+		CHAOS_API const Filter::FCombinedShapeFilterData GetCombinedShapeFilterData() const;
+		CHAOS_API void SetCombinedShapeFilterData(const Filter::FCombinedShapeFilterData& InCombinedShapeFilter);
+
 		FImplicitObjectRef GetGeometry() const;
 
 		const TAABB<FReal, 3>& GetWorldSpaceShapeBounds() const;
@@ -293,6 +300,13 @@ namespace Chaos
 		}
 		PRAGMA_ENABLE_INTERNAL_WARNINGS
 
+		CHAOS_API const Filter::FShapeFilterData GetShapeFilterData() const;
+		CHAOS_API void SetShapeFilterData(const Filter::FShapeFilterData& InShapeFilter);
+		CHAOS_API const Filter::FInstanceData GetFilterInstanceData() const;
+		CHAOS_API void SetFilterInstanceData(const Filter::FInstanceData& InInstanceData);
+		CHAOS_API const Filter::FCombinedShapeFilterData GetCombinedShapeFilterData() const;
+		CHAOS_API void SetCombinedShapeFilterData(const Filter::FCombinedShapeFilterData& InCombinedShapeFilter);
+
 		CHAOS_API void UpdateWorldSpaceState(const FRigidTransform3& WorldTransform);
 
 		// The leaf shape (with transformed and implicit wrapper removed).
@@ -566,6 +580,13 @@ namespace Chaos
 		void SetSimData(const FCollisionFilterData& InSimData) { CollisionData.SetSimData(InSimData); }
 		PRAGMA_ENABLE_INTERNAL_WARNINGS
 
+		CHAOS_API const Filter::FShapeFilterData GetShapeFilterData() const;
+		CHAOS_API void SetShapeFilterData(const Filter::FShapeFilterData& InShapeFilter);
+		CHAOS_API const Filter::FInstanceData GetFilterInstanceData() const;
+		CHAOS_API void SetFilterInstanceData(const Filter::FInstanceData& InInstanceData);
+		CHAOS_API const Filter::FCombinedShapeFilterData GetCombinedShapeFilterData() const;
+		CHAOS_API void SetCombinedShapeFilterData(const Filter::FCombinedShapeFilterData& InCombinedShapeFilter);
+
 		CHAOS_API void UpdateWorldSpaceState(const FRigidTransform3& WorldTransform);
 
 		// The leaf shape (with transformed and implicit wrapper removed).
diff --git a/Engine/Source/Runtime/Experimental/Chaos/Public/PhysicsProxy/GeometryCollectionPhysicsProxy.h b/Engine/Source/Runtime/Experimental/Chaos/Public/PhysicsProxy/GeometryCollectionPhysicsProxy.h
--- a/Engine/Source/Runtime/Experimental/Chaos/Public/PhysicsProxy/GeometryCollectionPhysicsProxy.h
+++ b/Engine/Source/Runtime/Experimental/Chaos/Public/PhysicsProxy/GeometryCollectionPhysicsProxy.h
@@ -511,7 +511,8 @@ public:
 	void SetReplicationMode(EReplicationMode Mode) { ReplicationMode = Mode; }
 	EReplicationMode GetReplicationMode() const { return ReplicationMode; }
 
-	CHAOS_API void UpdateFilterData_External(const FCollisionFilterData& NewSimFilter, const FCollisionFilterData& NewQueryFilter);
+	UE_DEPRECATED(5.7, "Use the overload that takes a FCombinedShapeFilterData instead") CHAOS_API void UpdateFilterData_External(const FCollisionFilterData& NewSimFilter, const FCollisionFilterData& NewQueryFilter);
+	CHAOS_API void UpdateFilterData_External(const Chaos::Filter::FCombinedShapeFilterData& NewShapeFilter);
 
 	struct FParticleCollisionFilterData
 	{
@@ -635,7 +636,8 @@ protected:
 	CHAOS_API void ScaleClusterGeometry_Internal(const FVector& WorldScale);
 
 	CHAOS_API void SetWorldTransform_Internal(const FTransform& WorldTransform, bool bInSkipChildToParentUpdateWhenInClusterUnion = false);
-	CHAOS_API void SetFilterData_Internal(const FCollisionFilterData& NewSimFilter, const FCollisionFilterData& NewQueryFilter);
+	UE_DEPRECATED(5.7, "Use the overload that takes a FCombinedShapeFilterData instead") CHAOS_API void SetFilterData_Internal(const FCollisionFilterData& NewSimFilter, const FCollisionFilterData& NewQueryFilter);
+	CHAOS_API void SetFilterData_Internal(const Chaos::Filter::FCombinedShapeFilterData& NewShapeFilter);
 	CHAOS_API void SetPerParticleFilterData_Internal(const TArray<FParticleCollisionFilterData>& PerParticleData);
 	CHAOS_API void SetDamagePropagationData_Internal(bool bEnabled, float BreakDamagePropagationFactor, float ShockDamagePropagationFactor);
 	CHAOS_API void SetEnableGravity_Internal(bool bEnabled);
diff --git a/Engine/Source/Runtime/Experimental/GeometryCollectionEngine/Private/GeometryCollection/GeometryCollectionComponent.cpp b/Engine/Source/Runtime/Experimental/GeometryCollectionEngine/Private/GeometryCollection/GeometryCollectionComponent.cpp
--- a/Engine/Source/Runtime/Experimental/GeometryCollectionEngine/Private/GeometryCollection/GeometryCollectionComponent.cpp
+++ b/Engine/Source/Runtime/Experimental/GeometryCollectionEngine/Private/GeometryCollection/GeometryCollectionComponent.cpp
@@ -4831,7 +4831,7 @@ void UGeometryCollectionComponent::LoadCollisionProfiles()
 
 				// Need to update on the GT too (physics proxy will just enqueue to do it on the PT) or else the changes won't take effect for GT SQ queries.
 				Interface->UpdateShapeCollisionFlags(ParticleView, ParticleData.bSimEnabled, ParticleData.bQueryEnabled);
-				Interface->UpdateShapeFilterData(ParticleView, ParticleData.QueryFilter, ParticleData.SimFilter);
+				Interface->UpdateCombinedShapeFilterData(ParticleView, Chaos::Filter::FShapeFilterBuilder::BuildFromLegacyShapeFilter(ParticleData.QueryFilter, ParticleData.SimFilter));
 			}
 		}
 	}
@@ -4927,7 +4927,7 @@ void UGeometryCollectionComponent::SetCollisionObjectType(ECollisionChannel Chan
 	// Update filters stored on proxy
 	if (PhysicsProxy)
 	{
-		PhysicsProxy->UpdateFilterData_External(InitialSimFilter, InitialQueryFilter);
+		PhysicsProxy->UpdateFilterData_External(Chaos::Filter::FShapeFilterBuilder::BuildFromLegacyShapeFilter(InitialQueryFilter, InitialSimFilter));
 	}
 }
 
@@ -4942,7 +4942,7 @@ void UGeometryCollectionComponent::OnActorEnableCollisionChanged()
 	// Update filters stored on proxy
 	if (PhysicsProxy)
 	{
-		PhysicsProxy->UpdateFilterData_External(InitialSimFilter, InitialQueryFilter);
+		PhysicsProxy->UpdateFilterData_External(Chaos::Filter::FShapeFilterBuilder::BuildFromLegacyShapeFilter(InitialQueryFilter, InitialSimFilter));
 	}
 
 }
diff --git a/Engine/Source/Runtime/PhysicsCore/Private/ChaosEngineInterface.cpp b/Engine/Source/Runtime/PhysicsCore/Private/ChaosEngineInterface.cpp
--- a/Engine/Source/Runtime/PhysicsCore/Private/ChaosEngineInterface.cpp
+++ b/Engine/Source/Runtime/PhysicsCore/Private/ChaosEngineInterface.cpp
@@ -2007,30 +2007,16 @@ FPhysicsGeometryCollection_Chaos FChaosEngineInterface::GetGeometryCollection(co
 
 void FChaosEngineInterface::SetMaskFilter(const FPhysicsShapeHandle& InShape, FMaskFilter InFilter)
 {
-	FCollisionFilterData SimFilter = GetSimulationFilter(InShape);
-	FCollisionFilterData QueryFilter = GetQueryFilter(InShape);
-
-	auto ApplyMask = [](uint32& Word3, FMaskFilter Mask)
-	{
-		// #CHAOSTODO - definitions for filter behavior are in the Engine module.
-		// Move all to PhysicsCore so we handle things in a safe way here.
-		static constexpr int32 LocalNumExtraBits = 6;
-		static_assert(LocalNumExtraBits <= 8, "Only up to 8 extra filter bits are supported.");
-		Word3 &= (0xFFFFFFFFu >> LocalNumExtraBits);	//we drop the top NumExtraFilterBits bits because that's where the new mask filter is going
-		Word3 |= uint32(Mask) << (32 - LocalNumExtraBits);
-	};
-	ApplyMask(SimFilter.Word3, InFilter);
-	ApplyMask(QueryFilter.Word3, InFilter);
-
-	SetSimulationFilter(InShape, SimFilter);
-	SetQueryFilter(InShape, QueryFilter);
+	Chaos::Filter::FShapeFilterData FilterData = GetShapeFilterData(InShape);
+	FilterData.SetMaskFilter(InFilter);
+	SetShapeFilterData(InShape, FilterData);
 }
 
 FCollisionFilterData FChaosEngineInterface::GetSimulationFilter(const FPhysicsShapeReference_Chaos& InShape)
 {
 	if(ensure(InShape.Shape))
 	{
-		return InShape.Shape->GetSimData();
+		return Chaos::Filter::FShapeFilterBuilder::GetLegacyShapeSimFilter(InShape.Shape->GetCombinedShapeFilterData());
 	} else
 	{
 		return FCollisionFilterData();
@@ -2041,21 +2027,76 @@ FCollisionFilterData FChaosEngineInterface::GetQueryFilter(const FPhysicsShapeRe
 {
 	if(ensure(InShape.Shape))
 	{
-		return InShape.Shape->GetQueryData();
+		return Chaos::Filter::FShapeFilterBuilder::GetLegacyShapeQueryFilter(InShape.Shape->GetCombinedShapeFilterData());
 	} else
 	{
 		return FCollisionFilterData();
 	}
 }
 
+Chaos::Filter::FShapeFilterData FChaosEngineInterface::GetShapeFilterData(const FPhysicsShapeHandle& InShape)
+{
+	if (ensure(InShape.Shape))
+	{
+		return InShape.Shape->GetShapeFilterData();
+	}
+	else
+	{
+		return Chaos::Filter::FShapeFilterData();
+	}
+}
+
+Chaos::Filter::FInstanceData FChaosEngineInterface::GetFilterInstanceData(const FPhysicsShapeHandle& InShape)
+{
+	if (ensure(InShape.Shape))
+	{
+		return InShape.Shape->GetFilterInstanceData();
+	}
+	else
+	{
+		return Chaos::Filter::FInstanceData();
+	}
+}
+
+Chaos::Filter::FCombinedShapeFilterData FChaosEngineInterface::GetCombinedShapeFilterData(const FPhysicsShapeHandle& InShape)
+{
+	if (ensure(InShape.Shape))
+	{
+		return InShape.Shape->GetCombinedShapeFilterData();
+	}
+	else
+	{
+		return Chaos::Filter::FCombinedShapeFilterData();
+	}
+}
+
 void FChaosEngineInterface::SetQueryFilter(const FPhysicsShapeReference_Chaos& InShapeRef,const FCollisionFilterData& InFilter)
 {
-	InShapeRef.Shape->SetQueryData(InFilter);
+	Chaos::Filter::FCombinedShapeFilterData CombinedShapeFilter = GetCombinedShapeFilterData(InShapeRef);
+	Chaos::Filter::FShapeFilterBuilder::SetLegacyShapeQueryFilter(CombinedShapeFilter, InFilter);
+	SetCombinedShapeFilterData(InShapeRef, CombinedShapeFilter);
 }
 
 void FChaosEngineInterface::SetSimulationFilter(const FPhysicsShapeReference_Chaos& InShapeRef,const FCollisionFilterData& InFilter)
 {
-	InShapeRef.Shape->SetSimData(InFilter);
+	Chaos::Filter::FCombinedShapeFilterData CombinedShapeFilter = GetCombinedShapeFilterData(InShapeRef);
+	Chaos::Filter::FShapeFilterBuilder::SetLegacyShapeSimFilter(CombinedShapeFilter, InFilter);
+	SetCombinedShapeFilterData(InShapeRef, CombinedShapeFilter);
+}
+
+void FChaosEngineInterface::SetShapeFilterData(const FPhysicsShapeHandle& InShapeRef, const Chaos::Filter::FShapeFilterData& InFilter)
+{
+	InShapeRef.Shape->SetShapeFilterData(InFilter);
+}
+
+void FChaosEngineInterface::SetFilterInstanceData(const FPhysicsShapeHandle& InShapeRef, const Chaos::Filter::FInstanceData& InFilter)
+{
+	InShapeRef.Shape->SetFilterInstanceData(InFilter);
+}
+
+void FChaosEngineInterface::SetCombinedShapeFilterData(const FPhysicsShapeReference_Chaos& InShapeRef, const Chaos::Filter::FCombinedShapeFilterData& InFilter)
+{
+	InShapeRef.Shape->SetCombinedShapeFilterData(InFilter);
 }
 
 bool FChaosEngineInterface::IsSimulationShape(const FPhysicsShapeHandle& InShape)
diff --git a/Engine/Source/Runtime/PhysicsCore/Private/ChaosInterfaceWrapperCore.cpp b/Engine/Source/Runtime/PhysicsCore/Private/ChaosInterfaceWrapperCore.cpp
--- a/Engine/Source/Runtime/PhysicsCore/Private/ChaosInterfaceWrapperCore.cpp
+++ b/Engine/Source/Runtime/PhysicsCore/Private/ChaosInterfaceWrapperCore.cpp
@@ -54,13 +54,27 @@ namespace ChaosInterface
 
 	FCollisionFilterData GetQueryFilterData(const Chaos::FPerShapeData& Shape)
 	{
-		return Shape.GetQueryData();
+		return Chaos::Filter::FShapeFilterBuilder::GetLegacyShapeQueryFilter(GetCombinedShapeFilterData(Shape));
 	}
 
 	FCollisionFilterData GetSimulationFilterData(const Chaos::FPerShapeData& Shape)
 	{
-		return Shape.GetSimData();
+		return Chaos::Filter::FShapeFilterBuilder::GetLegacyShapeSimFilter(GetCombinedShapeFilterData(Shape));
 	}
 
+	Chaos::Filter::FShapeFilterData GetShapeFilterData(const Chaos::FPerShapeData& Shape)
+	{
+		return Shape.GetShapeFilterData();
+	}
+
+	Chaos::Filter::FInstanceData GetFilterInstanceData(const Chaos::FPerShapeData& Shape)
+	{
+		return Shape.GetFilterInstanceData();
+	}
+
+	Chaos::Filter::FCombinedShapeFilterData GetCombinedShapeFilterData(const Chaos::FPerShapeData& Shape)
+	{
+		return Shape.GetCombinedShapeFilterData();
+	}
 
 }
diff --git a/Engine/Source/Runtime/PhysicsCore/Public/Chaos/ChaosEngineInterface.h b/Engine/Source/Runtime/PhysicsCore/Public/Chaos/ChaosEngineInterface.h
--- a/Engine/Source/Runtime/PhysicsCore/Public/Chaos/ChaosEngineInterface.h
+++ b/Engine/Source/Runtime/PhysicsCore/Public/Chaos/ChaosEngineInterface.h
@@ -554,8 +554,11 @@ public:
 	UE_DEPRECATED(5.4, "Use SetGeometry with FImplicitObjectPtr instead.")
 	static void SetGeometry(FPhysicsShapeHandle& InShape, TUniquePtr<Chaos::FImplicitObject>&& InGeometry) {check(false);}
 	
-	static PHYSICSCORE_API FCollisionFilterData GetSimulationFilter(const FPhysicsShapeHandle& InShape);
-	static PHYSICSCORE_API FCollisionFilterData GetQueryFilter(const FPhysicsShapeHandle& InShape);
+	UE_DEPRECATED(5.7, "Use GetCombinedShapeFilterData instead") static PHYSICSCORE_API FCollisionFilterData GetSimulationFilter(const FPhysicsShapeHandle& InShape);
+	UE_DEPRECATED(5.7, "Use GetCombinedShapeFilterData instead") static PHYSICSCORE_API FCollisionFilterData GetQueryFilter(const FPhysicsShapeHandle& InShape);
+	static PHYSICSCORE_API Chaos::Filter::FShapeFilterData GetShapeFilterData(const FPhysicsShapeHandle& InShape);
+	static PHYSICSCORE_API Chaos::Filter::FInstanceData GetFilterInstanceData(const FPhysicsShapeHandle& InShape);
+	static PHYSICSCORE_API Chaos::Filter::FCombinedShapeFilterData GetCombinedShapeFilterData(const FPhysicsShapeHandle& InShape);
 	static PHYSICSCORE_API bool IsSimulationShape(const FPhysicsShapeHandle& InShape);
 	static PHYSICSCORE_API bool IsQueryShape(const FPhysicsShapeHandle& InShape);
 	static PHYSICSCORE_API ECollisionShapeType GetShapeType(const FPhysicsShapeHandle& InShape);
@@ -571,8 +574,11 @@ public:
 	// @todo(mlentine): Which of these do we need to support?
 	// Set the mask filter of a shape, which is an extra level of filtering during collision detection / query for extra channels like "Blue Team" and "Red Team"
 	static PHYSICSCORE_API void SetMaskFilter(const FPhysicsShapeHandle& InShape,FMaskFilter InFilter);
-	static PHYSICSCORE_API void SetSimulationFilter(const FPhysicsShapeHandle& InShape,const FCollisionFilterData& InFilter);
-	static PHYSICSCORE_API void SetQueryFilter(const FPhysicsShapeHandle& InShape,const FCollisionFilterData& InFilter);
+	UE_DEPRECATED(5.7, "Use SetCombinedShapeFilterData instead") static PHYSICSCORE_API void SetSimulationFilter(const FPhysicsShapeHandle& InShape, const FCollisionFilterData& InFilter);
+	UE_DEPRECATED(5.7, "Use SetCombinedShapeFilterData instead") static PHYSICSCORE_API void SetQueryFilter(const FPhysicsShapeHandle& InShape, const FCollisionFilterData& InFilter);
+	static PHYSICSCORE_API void SetShapeFilterData(const FPhysicsShapeHandle& InShape, const Chaos::Filter::FShapeFilterData& InFilter);
+	static PHYSICSCORE_API void SetFilterInstanceData(const FPhysicsShapeHandle& InShape, const Chaos::Filter::FInstanceData& InFilter);
+	static PHYSICSCORE_API void SetCombinedShapeFilterData(const FPhysicsShapeHandle& InShape, const Chaos::Filter::FCombinedShapeFilterData& InFilter);
 	static PHYSICSCORE_API void SetIsSimulationShape(const FPhysicsShapeHandle& InShape,bool bIsSimShape);
 	static PHYSICSCORE_API void SetIsProbeShape(const FPhysicsShapeHandle& InShape,bool bIsProbeShape);
 	static PHYSICSCORE_API void SetIsQueryShape(const FPhysicsShapeHandle& InShape,bool bIsQueryShape);
diff --git a/Engine/Source/Runtime/PhysicsCore/Public/ChaosInterfaceWrapperCore.h b/Engine/Source/Runtime/PhysicsCore/Public/ChaosInterfaceWrapperCore.h
--- a/Engine/Source/Runtime/PhysicsCore/Public/ChaosInterfaceWrapperCore.h
+++ b/Engine/Source/Runtime/PhysicsCore/Public/ChaosInterfaceWrapperCore.h
@@ -93,8 +93,11 @@ struct FQueryDebugParams
 #endif
 };
 
-extern PHYSICSCORE_API FCollisionFilterData GetQueryFilterData(const Chaos::FPerShapeData& Shape);
-extern PHYSICSCORE_API FCollisionFilterData GetSimulationFilterData(const Chaos::FPerShapeData& Shape);
+UE_DEPRECATED(5.7, "Use GetCombinedShapeFilterData instead")  extern PHYSICSCORE_API FCollisionFilterData GetQueryFilterData(const Chaos::FPerShapeData& Shape);
+UE_DEPRECATED(5.7, "Use GetCombinedShapeFilterData instead")  extern PHYSICSCORE_API FCollisionFilterData GetSimulationFilterData(const Chaos::FPerShapeData& Shape);
+extern PHYSICSCORE_API Chaos::Filter::FShapeFilterData GetShapeFilterData(const Chaos::FPerShapeData& Shape);
+extern PHYSICSCORE_API Chaos::Filter::FInstanceData GetFilterInstanceData(const Chaos::FPerShapeData& Shape);
+extern PHYSICSCORE_API Chaos::Filter::FCombinedShapeFilterData GetCombinedShapeFilterData(const Chaos::FPerShapeData& Shape);
 
 
 PHYSICSCORE_API ECollisionShapeType GetImplicitType(const Chaos::FImplicitObject& InGeometry);
diff --git a/Engine/Source/Runtime/PhysicsCore/Public/PhysicsInterfaceUtilsCore.h b/Engine/Source/Runtime/PhysicsCore/Public/PhysicsInterfaceUtilsCore.h
--- a/Engine/Source/Runtime/PhysicsCore/Public/PhysicsInterfaceUtilsCore.h
+++ b/Engine/Source/Runtime/PhysicsCore/Public/PhysicsInterfaceUtilsCore.h
@@ -9,10 +9,10 @@
 
 namespace ChaosInterface { struct FPTLocationHit; }
 
-PHYSICSCORE_API FCollisionFilterData C2UFilterData(const FChaosFilterData& FilterData);
-PHYSICSCORE_API FChaosFilterData U2CFilterData(const FCollisionFilterData& FilterData);
+UE_INTERNAL PHYSICSCORE_API FCollisionFilterData C2UFilterData(const FChaosFilterData& FilterData);
+UE_INTERNAL PHYSICSCORE_API FChaosFilterData U2CFilterData(const FCollisionFilterData& FilterData);
 
-PHYSICSCORE_API FCollisionFilterData ToUnrealFilterData(const FChaosFilterData& FilterData);
+UE_INTERNAL PHYSICSCORE_API FCollisionFilterData ToUnrealFilterData(const FChaosFilterData& FilterData);
 
 //Find the face index for a given hit. This gives us a chance to modify face index based on things like most opposing normal
 
diff --git a/Engine/Source/Runtime/PhysicsCore/Public/SQCapture.h b/Engine/Source/Runtime/PhysicsCore/Public/SQCapture.h
--- a/Engine/Source/Runtime/PhysicsCore/Public/SQCapture.h
+++ b/Engine/Source/Runtime/PhysicsCore/Public/SQCapture.h
@@ -37,7 +37,7 @@ namespace Chaos
 }
 
 //Allows us to capture a scene query with either physx or chaos and then load it into either format for testing purposes
-struct FSQCapture
+struct UE_INTERNAL FSQCapture
 {
 	PHYSICSCORE_API ~FSQCapture();
 	FSQCapture(const FSQCapture&) = delete;
