Fix physics query redundantly computing some values and doing redundant null checks.
Made the visitor take the query geom's world aabb which is always computed on the outside.

diff --git a/Engine/Source/Runtime/PhysicsCore/Public/SQVisitor.h b/Engine/Source/Runtime/PhysicsCore/Public/SQVisitor.h
--- a/Engine/Source/Runtime/PhysicsCore/Public/SQVisitor.h
+++ b/Engine/Source/Runtime/PhysicsCore/Public/SQVisitor.h
@@ -42,11 +42,12 @@ template <typename QueryGeometryType, typename TPayload, typename THitType, bool
 struct TSQVisitor : public Chaos::ISpatialVisitor<TPayload, Chaos::FReal>
 {
 	using TGeometryType = std::conditional_t<bGTData, Chaos::FGeometryParticle, Chaos::FGeometryParticleHandle>;
+	using FAABB3 = Chaos::FAABB3;
+	using FVec3 = Chaos::FVec3;
 	TSQVisitor(const FVector& InStartPoint, const FVector& InDir, ChaosInterface::FSQHitBuffer<THitType>& InHitBuffer, EHitFlags InOutputFlags,
 		const ChaosInterface::FQueryFilterData& InQueryFilterData, ICollisionQueryFilterCallbackBase& InQueryCallback, const ChaosInterface::FQueryDebugParams& InDebugParams)
 		: StartPoint(InStartPoint)
 		, Dir(InDir)
-		, HalfExtents(0)
 		, OutputFlags(InOutputFlags)
 		, bAnyHit(false)
 		, DebugParams(InDebugParams)
@@ -55,14 +56,15 @@ struct TSQVisitor : public Chaos::ISpatialVisitor<TPayload, Chaos::FReal>
 		, QueryFilterData(InQueryFilterData)
 		, QueryFilterDataConcrete(C2UFilterData(QueryFilterData.data))
 		, QueryCallback(InQueryCallback)
+		, QueryCenter(StartTM.GetLocation())
 	{
+		QueryGeomWorldBounds = FAABB3(QueryCenter, QueryCenter);
 		bAnyHit = QueryFilterData.flags & FChaosQueryFlag::eANY_HIT;
 	}
 
 	TSQVisitor(const FTransform& InStartTM, const FVector& InDir, ChaosInterface::FSQHitBuffer<THitType>& InHitBuffer, EHitFlags InOutputFlags,
-		const ChaosInterface::FQueryFilterData& InQueryFilterData, ICollisionQueryFilterCallbackBase& InQueryCallback, const QueryGeometryType& InQueryGeom, const ChaosInterface::FQueryDebugParams& InDebugParams)
+		const ChaosInterface::FQueryFilterData& InQueryFilterData, ICollisionQueryFilterCallbackBase& InQueryCallback, const QueryGeometryType& InQueryGeom, const FAABB3& InQueryGeomWorldBounds, const ChaosInterface::FQueryDebugParams& InDebugParams)
 		: Dir(InDir)
-		, HalfExtents(InQueryGeom.BoundingBox().Extents() * 0.5)
 		, OutputFlags(InOutputFlags)
 		, bAnyHit(false)
 		, DebugParams(InDebugParams)
@@ -73,15 +75,23 @@ struct TSQVisitor : public Chaos::ISpatialVisitor<TPayload, Chaos::FReal>
 		, QueryGeom(&InQueryGeom)
 		, QueryCallback(InQueryCallback)
 		, StartTM(InStartTM)
+		, QueryGeomWorldBounds(InQueryGeomWorldBounds)
+		, QueryCenter(InQueryGeomWorldBounds.Center())
 	{
 		bAnyHit = QueryFilterData.flags & FChaosQueryFlag::eANY_HIT;
 		//todo: check THitType is sweep
 	}
 
-	TSQVisitor(const FTransform& InWorldTM, ChaosInterface::FSQHitBuffer<THitType>& InHitBuffer,
+	UE_DEPRECATED(5.8, "Use the overload that takes the QueryGeomWorldBounds")
+	TSQVisitor(const FTransform& InStartTM, const FVector& InDir, ChaosInterface::FSQHitBuffer<THitType>& InHitBuffer, EHitFlags InOutputFlags,
 		const ChaosInterface::FQueryFilterData& InQueryFilterData, ICollisionQueryFilterCallbackBase& InQueryCallback, const QueryGeometryType& InQueryGeom, const ChaosInterface::FQueryDebugParams& InDebugParams)
-		: HalfExtents(InQueryGeom.BoundingBox().Extents() * 0.5)
-		, bAnyHit(false)
+		: TSQVisitor(InStartTM, InDir, InHitBuffer, InOutputFlags, InQueryFilterData, InQueryCallback, InQueryGeom, InQueryGeom.CalculateTransformedBounds(InStartTM), InDebugParams)
+	{
+	}
+
+	TSQVisitor(const FTransform& InWorldTM, ChaosInterface::FSQHitBuffer<THitType>& InHitBuffer,
+		const ChaosInterface::FQueryFilterData& InQueryFilterData, ICollisionQueryFilterCallbackBase& InQueryCallback, const QueryGeometryType& InQueryGeom, const FAABB3& InQueryGeomWorldBounds, const ChaosInterface::FQueryDebugParams& InDebugParams)
+		: bAnyHit(false)
 		, DebugParams(InDebugParams)
 		, HitFaceNormal(Chaos::FVec3::ZeroVector)
 		, HitBuffer(InHitBuffer)
@@ -90,11 +100,20 @@ struct TSQVisitor : public Chaos::ISpatialVisitor<TPayload, Chaos::FReal>
 		, QueryGeom(&InQueryGeom)
 		, QueryCallback(InQueryCallback)
 		, StartTM(InWorldTM)
+		, QueryGeomWorldBounds(InQueryGeomWorldBounds)
+		, QueryCenter(InQueryGeomWorldBounds.Center())
 	{
 		bAnyHit = QueryFilterData.flags & FChaosQueryFlag::eANY_HIT;
 		//todo: check THitType is overlap
 	}
 
+	UE_DEPRECATED(5.8, "Use the overload that takes the QueryGeomWorldBounds")
+	TSQVisitor(const FTransform& InWorldTM, ChaosInterface::FSQHitBuffer<THitType>& InHitBuffer,
+		const ChaosInterface::FQueryFilterData& InQueryFilterData, ICollisionQueryFilterCallbackBase& InQueryCallback, const QueryGeometryType& InQueryGeom, const ChaosInterface::FQueryDebugParams& InDebugParams)
+		: TSQVisitor(InWorldTM, InHitBuffer, InQueryFilterData, InQueryCallback, InQueryGeom, InQueryGeom.CalculateTransformedBounds(InWorldTM), InDebugParams)
+	{
+	}
+
 	virtual ~TSQVisitor() {}
 
 	virtual bool Raycast(const Chaos::TSpatialVisitorData<TPayload>& Instance, Chaos::FQueryFastData& CurData) override
@@ -203,7 +222,6 @@ private:
 		bool bContinue = true;
 
 		const FRigidTransform3 ActorTM(GeometryParticle->GetX(), GeometryParticle->GetR());
-		const TAABB<FReal, 3> QueryGeomWorldBounds = QueryGeom ? QueryGeom->CalculateTransformedBounds(StartTM) : TAABB<FReal, 3>(-HalfExtents, HalfExtents);
 
 #if CHAOS_DEBUG_DRAW
 		bool bAllShapesIgnoredInPrefilter = true;
@@ -216,7 +234,7 @@ private:
 
 			const FImplicitObject* Geom = Shape->GetGeometry();
 
-			CVD_TRACE_SCOPED_SCENE_QUERY_VISIT_HELPER(EChaosVDSceneQueryVisitorType::NarrowPhase, FTransform(GeometryParticle->GetR(), GeometryParticle->GetX()), Payload.UniqueIdx().Idx, ShapeIndex, CurData);
+			CVD_TRACE_SCOPED_SCENE_QUERY_VISIT_HELPER(EChaosVDSceneQueryVisitorType::NarrowPhase, ActorTM, Payload.UniqueIdx().Idx, ShapeIndex, CurData);
 
 			if (bTestShapeBounds)
 			{
@@ -228,7 +246,7 @@ private:
 				else
 				{
 					// Transform to world bounds and get the proper half extent.
-					const FVec3 WorldHalfExtent = QueryGeom ? QueryGeomWorldBounds.Extents() * 0.5f : FVec3(HalfExtents);
+					const FVec3 WorldHalfExtent = QueryGeomWorldBounds.Extents() * 0.5f;
 
 					InflatedWorldBounds = FAABB3(Shape->GetWorldSpaceShapeBounds().Min() - WorldHalfExtent, Shape->GetWorldSpaceShapeBounds().Max() + WorldHalfExtent);
 				}
@@ -246,7 +264,6 @@ private:
 				}
 				else
 				{
-					const FVec3 QueryCenter = QueryGeom ? QueryGeomWorldBounds.Center() : StartTM.GetLocation();
 					if (!InflatedWorldBounds.Contains(QueryCenter))
 					{
 						CVD_SET_SQ_SHAPE_REJECT_REASON(EChaosVDSQVisitRejectReason::FailedFastBoundTest);
@@ -475,7 +492,6 @@ private:
 
 	const FVector StartPoint = FVector::ZeroVector;
 	const FVector Dir = FVector::ZeroVector;
-	const FVector HalfExtents = FVector::ZeroVector;
 	FHitFlags OutputFlags;
 	bool bAnyHit;
 	const ChaosInterface::FQueryDebugParams DebugParams;
@@ -486,6 +502,8 @@ private:
 	const QueryGeometryType* QueryGeom = nullptr;
 	ICollisionQueryFilterCallbackBase& QueryCallback;
 	const FTransform StartTM;
+	FAABB3 QueryGeomWorldBounds = FAABB3::ZeroAABB();
+	FVec3 QueryCenter = FVec3::ZeroVector;
 };
 
 
@@ -643,7 +661,7 @@ void SweepHelper(const QueryGeomType& QueryGeom, const Chaos::ISpatialAccelerati
 	const bool bIsValid = !Bounds.Min().ContainsNaN() && !Bounds.Max().ContainsNaN();
 	if (ensure(bIsValid))
 	{
-		TSQVisitor<QueryGeomType, TPayload, TSweepHit, std::is_same<TSweepHit, FSweepHit>::value> SweepVisitor(StartTM, Dir, HitBuffer, OutputFlags, QueryFilterData, QueryCallback, QueryGeom, DebugParams);
+		TSQVisitor<QueryGeomType, TPayload, TSweepHit, std::is_same<TSweepHit, FSweepHit>::value> SweepVisitor(StartTM, Dir, HitBuffer, OutputFlags, QueryFilterData, QueryCallback, QueryGeom, Bounds, DebugParams);
 
 		HitBuffer.IncFlushCount();
 
@@ -680,7 +698,7 @@ void OverlapHelper(const QueryGeomType& QueryGeom, const Chaos::ISpatialAccelera
 		}
 		else
 		{
-			TSQVisitor<QueryGeomType, TPayload, TOverlapHit, bGTData> OverlapVisitor(GeomPose, HitBuffer, QueryFilterData, QueryCallback, QueryGeom, DebugParams);
+			TSQVisitor<QueryGeomType, TPayload, TOverlapHit, bGTData> OverlapVisitor(GeomPose, HitBuffer, QueryFilterData, QueryCallback, QueryGeom, Bounds, DebugParams);
 			SpatialAcceleration.Overlap(Bounds, OverlapVisitor);
 		}
 		HitBuffer.DecFlushCount();
