Chaos: Introducing FrozenDuringResim ResimType enum value and making that the default value, so only particle in a resim bubble have to be marked as FullResim.

diff --git a/Engine/Source/Runtime/Experimental/Chaos/Public/Chaos/GeometryParticles.h b/Engine/Source/Runtime/Experimental/Chaos/Public/Chaos/GeometryParticles.h
--- a/Engine/Source/Runtime/Experimental/Chaos/Public/Chaos/GeometryParticles.h
+++ b/Engine/Source/Runtime/Experimental/Chaos/Public/Chaos/GeometryParticles.h
@@ -142,9 +142,9 @@ namespace Chaos
 	enum class EResimType : uint8
 	{
 		FullResim = 0,	//fully re-run simulation and keep results (any forces must be applied again)
+		ResimAsFollower = 1, //use previous forces and snap to previous results regardless of variation - used to push other objects away
+		FrozenDuringResim = 2, // this particle is "frozen" (made kinematic), updates and integration should be skipped
 		//ResimWithPrevForces, //use previous forces and keep results (UNIMPLEMENTED)
-		ResimAsFollower = 1 //use previous forces and snap to previous results regardless of variation - used to push other objects away
-		//ResimAsKinematic //treat as kinematic (UNIMPLEMENTED)
 	};
 	
 	template<class T, int d, EGeometryParticlesSimType SimType>
diff --git a/Engine/Source/Runtime/Experimental/Chaos/Public/Chaos/ParticleHandle.h b/Engine/Source/Runtime/Experimental/Chaos/Public/Chaos/ParticleHandle.h
--- a/Engine/Source/Runtime/Experimental/Chaos/Public/Chaos/ParticleHandle.h
+++ b/Engine/Source/Runtime/Experimental/Chaos/Public/Chaos/ParticleHandle.h
@@ -86,7 +86,7 @@ void GeometryParticleDefaultConstruct(FConcrete& Concrete, const FGeometryPartic
 	Concrete.SetX(TVector<T, d>(0));
 	Concrete.SetR(TRotation<T, d>::Identity);
 	Concrete.SetSpatialIdx(FSpatialAccelerationIdx{ 0,0 });
-	Concrete.SetResimType(EResimType::FullResim);
+	Concrete.SetResimType(EResimType::FrozenDuringResim);
 	Concrete.SetEnabledDuringResim(true);
 }
 
diff --git a/Engine/Source/Runtime/Experimental/Chaos/Public/PhysicsProxy/SingleParticlePhysicsProxy.h b/Engine/Source/Runtime/Experimental/Chaos/Public/PhysicsProxy/SingleParticlePhysicsProxy.h
--- a/Engine/Source/Runtime/Experimental/Chaos/Public/PhysicsProxy/SingleParticlePhysicsProxy.h
+++ b/Engine/Source/Runtime/Experimental/Chaos/Public/PhysicsProxy/SingleParticlePhysicsProxy.h
@@ -636,6 +636,19 @@ public:
 			});
 	}
 
+	EResimType ResimType() const
+	{
+		return Read([](auto* Particle)
+			{
+				if (auto Rigid = Particle->CastToRigidParticle())
+				{
+					return Rigid->ResimType();
+				}
+
+				return EResimType::FrozenDuringResim;
+			});
+	}
+
 	void SetResimType(EResimType ResimType)
 	{
 		Write([ResimType](auto* Particle)
@@ -647,16 +660,6 @@ public:
 		});
 	}
 
-	EResimType ResimType() const
-	{
-		if (auto Rigid = Particle->CastToRigidParticle())
-		{
-			return Rigid->ResimType();
-		}
-
-		return EResimType::FullResim;
-	}
-
 	const FVec3 Acceleration() const
 	{
 		return Read([](auto* Particle)
diff --git a/Engine/Source/Runtime/PhysicsCore/Private/ChaosEngineInterface.cpp b/Engine/Source/Runtime/PhysicsCore/Private/ChaosEngineInterface.cpp
--- a/Engine/Source/Runtime/PhysicsCore/Private/ChaosEngineInterface.cpp
+++ b/Engine/Source/Runtime/PhysicsCore/Private/ChaosEngineInterface.cpp
@@ -529,13 +529,7 @@ void FChaosEngineInterface::SetIsKinematic_AssumesLocked(const FPhysicsActorHand
 		if(AllowedToChangeToNewState)
 		{
 			InActorReference->GetGameThreadAPI().SetObjectState(NewState);
-			//we mark as full resim only if going from kinematic to simulated
-			//going from simulated to kinematic we assume user is doing some optimization so we leave it up to them
-			if(NewState == EObjectStateType::Dynamic)
-			{
-				InActorReference->GetGameThreadAPI().SetResimType(EResimType::FullResim);
-			}
-			else if (NewState == Chaos::EObjectStateType::Kinematic)
+			if (NewState == Chaos::EObjectStateType::Kinematic)
 			{
 				// Reset velocity on a state change here
 				InActorReference->GetGameThreadAPI().SetV(Chaos::FVec3((Chaos::FReal) 0));
@@ -2426,7 +2420,6 @@ void FChaosEngineInterface::CreateActor(const FActorCreationParams& InParams,FPh
 	if(InParams.bStatic)
 	{
 		Particle = FGeometryParticle::CreateParticle();
-		Particle->SetResimType(EResimType::ResimAsFollower);
 	}
 	else
 	{
@@ -2443,15 +2436,14 @@ void FChaosEngineInterface::CreateActor(const FActorCreationParams& InParams,FPh
 			{
 				Rigid->SetObjectState(EObjectStateType::Sleeping);
 			}
-			Rigid->SetResimType(EResimType::FullResim);
 		} else
 		{
 			Rigid->SetObjectState(EObjectStateType::Kinematic);
-			Rigid->SetResimType(EResimType::ResimAsFollower);	//for now kinematics are never changed during resim
 		}
 		//Particle.Reset(Rigid.Release());
 		Particle = MoveTemp(Rigid);
 	}
+	Particle->SetResimType(EResimType::FrozenDuringResim);
 
 	// Set the particle acceleration structure spatial index here
 	{
