Collision Filter Refactor Part 3.13: Large deprecation pass focusing on CollisionFilterQueryCallback.
This fixes the callstack from SceneQuery through the SQVisitor and ending up at CollisionFilterQueryCallback.
This deprecates some old classes and structures as well.

diff --git a/Engine/Plugins/Experimental/SpatialReadiness/Source/SpatialReadiness/Private/SpatialReadiness/SpatialReadinessSimCallback.cpp b/Engine/Plugins/Experimental/SpatialReadiness/Source/SpatialReadiness/Private/SpatialReadiness/SpatialReadinessSimCallback.cpp
--- a/Engine/Plugins/Experimental/SpatialReadiness/Source/SpatialReadiness/Private/SpatialReadiness/SpatialReadinessSimCallback.cpp
+++ b/Engine/Plugins/Experimental/SpatialReadiness/Source/SpatialReadiness/Private/SpatialReadiness/SpatialReadinessSimCallback.cpp
@@ -244,11 +244,9 @@ bool FSpatialReadinessSimCallback::QueryReadiness_GT(const FBox& Bounds, TArray<
 	QueryParams.bTraceComplex = false;
 	FCollisionObjectQueryParams ObjectParams(static_cast<uint8>(Channel));
 	FCollisionResponseContainer ResponseContainer;
-	FCollisionFilterData CollisionFilterData = Chaos::Filter::FQueryFilterBuilder::GetLegacyQueryFilter(CreateChaosQueryFilterData(ChannelBit, bComplex, ResponseContainer, QueryParams, ObjectParams, bMulti));
-	FOverlapAllQueryCallback QueryCallback;
+	const Chaos::Filter::FQueryFilterData CollisionFilterData = CreateChaosQueryFilterData(ChannelBit, bComplex, ResponseContainer, QueryParams, ObjectParams, bMulti);
+	FOverlapAllCollisionQueryFilterCallback QueryCallback;
 	EQueryFlags QueryFlags = EQueryFlags::PreFilter;// | EQueryFlags::SkipNarrowPhase;
-	ChaosInterface::FQueryFilterData QueryFilterData = ChaosInterface::MakeQueryFilterData(CollisionFilterData, QueryFlags, QueryParams);
-	ChaosInterface::FQueryDebugParams DebugParams = { };
 
 	// Create the geometry needed for the query
 	Chaos::FImplicitBox3 Geom = Chaos::FImplicitBox3(Bounds.Min, Bounds.Max, 0);
@@ -256,9 +254,10 @@ bool FSpatialReadinessSimCallback::QueryReadiness_GT(const FBox& Bounds, TArray<
 	// Do the query
 	FPhysicsHitCallback<ChaosInterface::FOverlapHit> HitBuffer;
 	FPhysicsCommand::ExecuteRead(&PhysicsScene,
-	[this, &Geom, &HitBuffer, &QueryFlags, &CollisionFilterData, &QueryFilterData, &QueryCallback, &DebugParams]()
+	[this, &Geom, &HitBuffer, &QueryFlags, &CollisionFilterData, &QueryCallback]()
 	{
-		Chaos::Private::LowLevelOverlap(PhysicsScene, Geom, FTransform::Identity, HitBuffer, QueryFlags, CollisionFilterData, QueryFilterData, &QueryCallback, DebugParams);
+		const ChaosInterface::FSceneQueryCommonParams CommonParams(QueryCallback, CollisionFilterData, QueryFlags);
+		Chaos::Private::LowLevelOverlap(PhysicsScene, Geom, FTransform::Identity, HitBuffer, CommonParams);
 	});
 
 	// Make sure the hits from the buffer are actually in our list of unready volumes,
diff --git a/Engine/Source/Programs/HeadlessChaos/Private/HeadlessChaosTestPhysicsScene.cpp b/Engine/Source/Programs/HeadlessChaos/Private/HeadlessChaosTestPhysicsScene.cpp
--- a/Engine/Source/Programs/HeadlessChaos/Private/HeadlessChaosTestPhysicsScene.cpp
+++ b/Engine/Source/Programs/HeadlessChaos/Private/HeadlessChaosTestPhysicsScene.cpp
@@ -138,8 +138,9 @@ namespace ChaosTest {
 	{
 		FChaosSQAccelerator SQAccelerator(*Scene.GetSpacialAcceleration());
 		FSQHitBuffer<ChaosInterface::FOverlapHit> HitBuffer;
-		FOverlapAllQueryCallback QueryCallback;
-		SQAccelerator.Overlap(Chaos::FSphere(FVec3(0), Radius), InTM, HitBuffer, FChaosQueryFilterData(), QueryCallback, FQueryDebugParams());
+		FOverlapAllCollisionQueryFilterCallback QueryCallback;
+		ChaosInterface::FSceneQueryCommonParams CommonParams(QueryCallback);
+		SQAccelerator.Overlap(Chaos::FSphere(FVec3(0), Radius), InTM, HitBuffer, CommonParams);
 		return HitBuffer;
 	}
 
@@ -1957,14 +1958,15 @@ namespace ChaosTest {
 
 		FChaosSQAccelerator SQ{ *Scene.GetSpacialAcceleration() };
 		FSQHitBuffer<ChaosInterface::FOverlapHit> HitBuffer;
-		FOverlapAllQueryCallback QueryCallback;
+		FOverlapAllCollisionQueryFilterCallback QueryCallback;
+		ChaosInterface::FSceneQueryCommonParams CommonParams(QueryCallback);
 
 		// Here we query from a position under the box, but using a shape that has an offset. This tests
 		// a failure case that was previously present where the query system assumed that the QueryTM
 		// was inside the geometry being used to query.
 		const FTransform QueryTM{ FVec3{0.0f, 0.0f, -110.0f} };
 		constexpr FReal SphereRadius = static_cast<FReal>(50.0);
-		SQ.Overlap(Chaos::FSphere(FVec3(0.0f, 0.0f, 100.0f), SphereRadius), QueryTM, HitBuffer, FChaosQueryFilterData(), QueryCallback, FQueryDebugParams());
+		SQ.Overlap(Chaos::FSphere(FVec3(0.0f, 0.0f, 100.0f), SphereRadius), QueryTM, HitBuffer, CommonParams);
 
 		EXPECT_TRUE(HitBuffer.HasBlockingHit());
 	}
@@ -2007,7 +2009,8 @@ namespace ChaosTest {
 
 		FChaosSQAccelerator SQ{ *Scene.GetSpacialAcceleration() };
 		FSQHitBuffer<ChaosInterface::FSweepHit> HitBuffer;
-		FBlockAllQueryCallback QueryCallback;
+		FBlockAllCollisionQueryFilterCallback QueryCallback;
+		ChaosInterface::FSceneQueryCommonParams CommonParams(QueryCallback);
 
 		// Another box of same size that is offset from origin by 200.
 		FVec3 Offset(200.f,0,0);
@@ -2017,7 +2020,7 @@ namespace ChaosTest {
 		const FTransform QueryTM{ FVec3{-200.f, 0, 100.0f} };
 		const FVec3 Dir(0,0,-1);
 		const FReal Length = 200;
-		SQ.Sweep(QueryBox, QueryTM, Dir, Length, HitBuffer, EHitFlags::None, FQueryFilterData(), QueryCallback, FQueryDebugParams());
+		SQ.Sweep(QueryBox, QueryTM, Dir, Length, HitBuffer, EHitFlags::None, CommonParams);
 
 		EXPECT_TRUE(HitBuffer.HasBlockingHit());
 	}
@@ -2053,7 +2056,8 @@ namespace ChaosTest {
 
 		FChaosSQAccelerator SQ{ *Scene.GetSpacialAcceleration() };
 		FSQHitBuffer<ChaosInterface::FSweepHit> HitBuffer;
-		FBlockAllQueryCallback QueryCallback;
+		FBlockAllCollisionQueryFilterCallback QueryCallback;
+		ChaosInterface::FSceneQueryCommonParams CommonParams(QueryCallback);
 
 		const FReal SphereRadius = 25;
 		const Chaos::FSphere QuerySphere(FVec3::ZeroVector, SphereRadius);
@@ -2062,7 +2066,7 @@ namespace ChaosTest {
 		FTransform QueryTM{ FVec3{-5, 0, 50} };
 		const FVec3 Dir(0,0,-1);
 		const FReal Length = 200;
-		SQ.Sweep(QuerySphere, QueryTM, Dir, Length, HitBuffer, EHitFlags::MTD, FQueryFilterData(), QueryCallback, FQueryDebugParams());
+		SQ.Sweep(QuerySphere, QueryTM, Dir, Length, HitBuffer, EHitFlags::MTD, CommonParams);
 
 		EXPECT_TRUE(HitBuffer.HasBlockingHit());
 		const FSweepHit* BlockingHit = HitBuffer.GetBlock(); 
@@ -2073,7 +2077,7 @@ namespace ChaosTest {
 
 		// Now do a second sweep down, this time over shape 1 more.
 		QueryTM = FTransform{ FVec3{5, 0, 50} };
-		SQ.Sweep(QuerySphere, QueryTM, Dir, Length, HitBuffer, EHitFlags::MTD, FQueryFilterData(), QueryCallback, FQueryDebugParams());
+		SQ.Sweep(QuerySphere, QueryTM, Dir, Length, HitBuffer, EHitFlags::MTD, CommonParams);
 
 		EXPECT_TRUE(HitBuffer.HasBlockingHit());
 		BlockingHit = HitBuffer.GetBlock();
@@ -2118,7 +2122,8 @@ namespace ChaosTest {
 
 		FChaosSQAccelerator SQ{ *Scene.GetSpacialAcceleration() };
 		FSQHitBuffer<ChaosInterface::FSweepHit> HitBuffer;
-		FBlockAllQueryCallback QueryCallback;
+		FBlockAllCollisionQueryFilterCallback QueryCallback;
+		ChaosInterface::FSceneQueryCommonParams CommonParams(QueryCallback);
 
 		const FReal SphereRadius = 25;
 		const Chaos::FSphere QuerySphere(FVec3::ZeroVector, SphereRadius);
@@ -2127,7 +2132,7 @@ namespace ChaosTest {
 		FTransform QueryTM{ FVec3{-5, 0, 50} };
 		const FVec3 Dir(0, 0, -1);
 		const FReal Length = 200;
-		SQ.Sweep(QuerySphere, QueryTM, Dir, Length, HitBuffer, EHitFlags::MTD, FQueryFilterData(), QueryCallback, FQueryDebugParams());
+		SQ.Sweep(QuerySphere, QueryTM, Dir, Length, HitBuffer, EHitFlags::MTD, CommonParams);
 
 		EXPECT_TRUE(HitBuffer.HasBlockingHit());
 		const FSweepHit* BlockingHit = HitBuffer.GetBlock();
@@ -2138,7 +2143,7 @@ namespace ChaosTest {
 
 		// Now do a second sweep down, this time over shape 1 more.
 		QueryTM = FTransform{ FVec3{5, 0, 50} };
-		SQ.Sweep(QuerySphere, QueryTM, Dir, Length, HitBuffer, EHitFlags::MTD, FQueryFilterData(), QueryCallback, FQueryDebugParams());
+		SQ.Sweep(QuerySphere, QueryTM, Dir, Length, HitBuffer, EHitFlags::MTD, CommonParams);
 
 		EXPECT_TRUE(HitBuffer.HasBlockingHit());
 		BlockingHit = HitBuffer.GetBlock();
diff --git a/Engine/Source/Runtime/Engine/Private/Collision/SceneQuery.cpp b/Engine/Source/Runtime/Engine/Private/Collision/SceneQuery.cpp
--- a/Engine/Source/Runtime/Engine/Private/Collision/SceneQuery.cpp
+++ b/Engine/Source/Runtime/Engine/Private/Collision/SceneQuery.cpp
@@ -290,28 +290,16 @@ struct TSQTraits
 
 	//SceneTrace - ray
 	template <typename AccelContainerType, typename GeomInputsType, ESweepOrRay T = GeometryQuery>
-	static typename TEnableIf<T == ESweepOrRay::Raycast, void>::Type SceneTrace(const AccelContainerType& Container, const GeomInputsType& GeomInputs, const FVector& Dir, float DeltaMag, const FTransform& StartTM, HitBufferType& HitBuffer, EHitFlags OutputFlags, EQueryFlags QueryFlags, const FCollisionFilterData& FilterData, const FCollisionQueryParams& Params, ICollisionQueryFilterCallbackBase* QueryCallback)
+	static typename TEnableIf<T == ESweepOrRay::Raycast, void>::Type SceneTrace(const AccelContainerType& Container, const GeomInputsType& GeomInputs, const FVector& Dir, float DeltaMag, const FTransform& StartTM, HitBufferType& HitBuffer, EHitFlags OutputFlags, const ChaosInterface::FSceneQueryCommonParams& CommonParams)
 	{
-		using namespace ChaosInterface;
-		FQueryFilterData QueryFilterData = MakeQueryFilterData(FilterData, QueryFlags, Params);
-		FQueryDebugParams DebugParams;
-#if !(UE_BUILD_TEST || UE_BUILD_SHIPPING)
-		DebugParams.bDebugQuery = Params.bDebugQuery;
-#endif
-		Chaos::Private::LowLevelRaycast(Container, StartTM.GetLocation(), Dir, DeltaMag, HitBuffer, OutputFlags, QueryFlags, FilterData, QueryFilterData, QueryCallback, DebugParams);	//todo(ocohen): namespace?
+		Chaos::Private::LowLevelRaycast(Container, StartTM.GetLocation(), Dir, DeltaMag, HitBuffer, OutputFlags, CommonParams);
 	}
 
 	//SceneTrace - sweep
 	template <typename AccelContainerType, typename GeomInputsType, ESweepOrRay T = GeometryQuery>
-	static typename TEnableIf<T == ESweepOrRay::Sweep, void>::Type SceneTrace(const AccelContainerType& Container, const GeomInputsType& GeomInputs, const FVector& Dir, float DeltaMag, const FTransform& StartTM, HitBufferType& HitBuffer, EHitFlags OutputFlags, EQueryFlags QueryFlags, const FCollisionFilterData& FilterData, const FCollisionQueryParams& Params, ICollisionQueryFilterCallbackBase* QueryCallback)
+	static typename TEnableIf<T == ESweepOrRay::Sweep, void>::Type SceneTrace(const AccelContainerType& Container, const GeomInputsType& GeomInputs, const FVector& Dir, float DeltaMag, const FTransform& StartTM, HitBufferType& HitBuffer, EHitFlags OutputFlags, const ChaosInterface::FSceneQueryCommonParams& CommonParams)
 	{
-		using namespace ChaosInterface;
-		FQueryFilterData QueryFilterData = MakeQueryFilterData(FilterData, QueryFlags, Params);
-		FQueryDebugParams DebugParams;
-#if !(UE_BUILD_TEST || UE_BUILD_SHIPPING)
-		DebugParams.bDebugQuery = Params.bDebugQuery;
-#endif
-		Chaos::Private::LowLevelSweep(Container, *GeomInputs.GetGeometry(), StartTM, Dir, DeltaMag, HitBuffer, OutputFlags, QueryFlags, FilterData, QueryFilterData, QueryCallback, DebugParams);	//todo(ocohen): namespace?
+		Chaos::Private::LowLevelSweep(Container, *GeomInputs.GetGeometry(), StartTM, Dir, DeltaMag, HitBuffer, OutputFlags, CommonParams);
 	}
 
 	static void ResetOutHits(TArray<FHitResult>& OutHits, const FVector& Start, const FVector& End)
@@ -529,7 +517,6 @@ bool TSceneCastCommonImpWithRetryRequest(const UWorld* World, typename Traits::O
 		// Create filter data used to filter collisions 
 		CA_SUPPRESS(6326);
 		const Chaos::Filter::FQueryFilterData QueryFilterData = CreateChaosQueryFilterData(TraceChannel, Params.bTraceComplex, ResponseParams.CollisionResponse, Params, ObjectParams, Traits::SingleMultiOrTest == ESingleMultiOrTest::Multi);
-		FCollisionFilterData Filter = Chaos::Filter::FQueryFilterBuilder::GetLegacyQueryFilter(QueryFilterData);
 		
 		CA_SUPPRESS(6326);
 		FCollisionQueryFilterCallback QueryCallback(Params, Traits::GeometryQuery == ESweepOrRay::Sweep);
@@ -546,21 +533,27 @@ bool TSceneCastCommonImpWithRetryRequest(const UWorld* World, typename Traits::O
 		const FVector Dir = DeltaMag > 0.f ? (Delta / DeltaMag) : FVector(1, 0, 0);
 		const FTransform StartTM = Traits::IsRay() ? FTransform(Start) : FTransform(*GeomInputs.GetGeometryOrientation(), Start);
 
+		FQueryDebugParams DebugParams;
+#if !(UE_BUILD_TEST || UE_BUILD_SHIPPING)
+		DebugParams.bDebugQuery = Params.bDebugQuery;
+#endif
+
 		// Enable scene locks, in case they are required
 		FPhysScene& PhysScene = *World->GetPhysicsScene();
 		
 		FScopedSceneReadLock SceneLocks(PhysScene);
 		{
+			const ChaosInterface::FSceneQueryCommonParams CommonParams(QueryCallback, QueryFilterData, Traits::GetQueryFlags(), DebugParams);
 			FScopedSQHitchRepeater<decltype(HitBufferSync)> HitchRepeater(HitBufferSync, QueryCallback, FHitchDetectionInfo(Start, End, TraceChannel, Params));
 			do
 			{
 				if constexpr (AccelContainerType::HasAccelerationStructureOverride())
 				{
-					Traits::SceneTrace(AccelContainer.GetSpatialAcceleration(), GeomInputs, Dir, DeltaMag, StartTM, HitchRepeater.GetBuffer(), Traits::GetHitFlags(), Traits::GetQueryFlags(), Filter, Params, &QueryCallback);
+					Traits::SceneTrace(AccelContainer.GetSpatialAcceleration(), GeomInputs, Dir, DeltaMag, StartTM, HitchRepeater.GetBuffer(), Traits::GetHitFlags(), CommonParams);
 				}
 				else
 				{
-					Traits::SceneTrace(PhysScene, GeomInputs, Dir, DeltaMag, StartTM, HitchRepeater.GetBuffer(), Traits::GetHitFlags(), Traits::GetQueryFlags(), Filter, Params, &QueryCallback);
+					Traits::SceneTrace(PhysScene, GeomInputs, Dir, DeltaMag, StartTM, HitchRepeater.GetBuffer(), Traits::GetHitFlags(), CommonParams);
 				}
 			} while (HitchRepeater.RepeatOnHitch());
 		}
@@ -1117,7 +1110,6 @@ bool GeomOverlapMultiImp(const UWorld* World, const FPhysicsGeometry& Geom, cons
 
 		// Create filter data used to filter collisions
 		const Chaos::Filter::FQueryFilterData QueryFilterData = CreateChaosQueryFilterData(TraceChannel, Params.bTraceComplex, ResponseParams.CollisionResponse, Params, ObjectParams, InfoType != EQueryInfo::IsAnything);
-		FCollisionFilterData Filter = Chaos::Filter::FQueryFilterBuilder::GetLegacyQueryFilter(QueryFilterData);
 		FCollisionQueryFilterCallback QueryCallback(Params, false);
 		QueryCallback.bIgnoreTouches |= (InfoType == EQueryInfo::IsBlocking); // pre-filter to ignore touches and only get blocking hits, if that's what we're after.
 		QueryCallback.bIsOverlapQuery = true;
@@ -1139,17 +1131,18 @@ bool GeomOverlapMultiImp(const UWorld* World, const FPhysicsGeometry& Geom, cons
 
 		FPhysicsCommand::ExecuteRead(&PhysScene, [&]()
 		{
+			const FSceneQueryCommonParams CommonParams(QueryCallback, QueryFilterData, QueryFlags, DebugParams);
 			{
 				FScopedSQHitchRepeater<typename TRemoveReference<decltype(OverlapBuffer)>::Type> HitchRepeater(OverlapBuffer, QueryCallback, FHitchDetectionInfo(GeomPose, TraceChannel, Params));
 				do
 				{
 					if constexpr (AccelContainerType::HasAccelerationStructureOverride())
 					{
-						Chaos::Private::LowLevelOverlap(AccelContainer.GetSpatialAcceleration(), Geom, GeomPose, HitchRepeater.GetBuffer(), QueryFlags, Filter, MakeQueryFilterData(Filter, QueryFlags, Params), &QueryCallback, DebugParams);
+						Chaos::Private::LowLevelOverlap(AccelContainer.GetSpatialAcceleration(), Geom, GeomPose, HitchRepeater.GetBuffer(), CommonParams);
 					}
 					else
 					{
-						Chaos::Private::LowLevelOverlap(PhysScene, Geom, GeomPose, HitchRepeater.GetBuffer(), QueryFlags, Filter, MakeQueryFilterData(Filter, QueryFlags, Params), &QueryCallback, DebugParams);
+						Chaos::Private::LowLevelOverlap(PhysScene, Geom, GeomPose, HitchRepeater.GetBuffer(), CommonParams);
 					}
 				} while(HitchRepeater.RepeatOnHitch());
 
diff --git a/Engine/Source/Runtime/Engine/Private/Collision/SceneQueryLowLevel.cpp b/Engine/Source/Runtime/Engine/Private/Collision/SceneQueryLowLevel.cpp
--- a/Engine/Source/Runtime/Engine/Private/Collision/SceneQueryLowLevel.cpp
+++ b/Engine/Source/Runtime/Engine/Private/Collision/SceneQueryLowLevel.cpp
@@ -25,32 +25,32 @@
 		{}
 
 		template<typename THit>
-		void Raycast(const FVector& Start, const FVector& Dir, const float DeltaMagnitude, ChaosInterface::FSQHitBuffer<THit>& HitBuffer, EHitFlags OutputFlags, const ChaosInterface::FQueryFilterData& QueryFilterData, ICollisionQueryFilterCallbackBase& QueryCallback, const ChaosInterface::FQueryDebugParams& DebugParams = {}) const
+		void Raycast(const FVector& Start, const FVector& Dir, const float DeltaMagnitude, ChaosInterface::FSQHitBuffer<THit>& HitBuffer, EHitFlags OutputFlags, const ChaosInterface::FSceneQueryCommonParams& CommonParams) const
 		{
 			using namespace Chaos;
 			using namespace ChaosInterface;
 
-			TSQVisitor<TSphere<FReal, 3>, TPayload, THit, std::is_same_v<THit, FRaycastHit>> RaycastVisitor(Start, Dir, HitBuffer, OutputFlags, QueryFilterData, QueryCallback, DebugParams);
+			TSQVisitor<TSphere<FReal, 3>, TPayload, THit, std::is_same_v<THit, FRaycastHit>> RaycastVisitor(Start, Dir, HitBuffer, OutputFlags, CommonParams);
 			HitBuffer.IncFlushCount();
 			SpatialAcceleration.Raycast(Start, Dir, DeltaMagnitude, RaycastVisitor);
 			HitBuffer.DecFlushCount();
 		}
 
 		template<typename THit>
-		void Sweep(const Chaos::FImplicitObject& QueryGeom, const FTransform& StartTM, const FVector& Dir, const float DeltaMagnitude, ChaosInterface::FSQHitBuffer<THit>& HitBuffer, EHitFlags OutputFlags, const ChaosInterface::FQueryFilterData& QueryFilterData, ICollisionQueryFilterCallbackBase& QueryCallback, const ChaosInterface::FQueryDebugParams& DebugParams = {}) const
+		void Sweep(const Chaos::FImplicitObject& QueryGeom, const FTransform& StartTM, const FVector& Dir, const float DeltaMagnitude, ChaosInterface::FSQHitBuffer<THit>& HitBuffer, EHitFlags OutputFlags, const ChaosInterface::FSceneQueryCommonParams& CommonParams) const
 		{
 			return Chaos::Utilities::CastHelper(QueryGeom, StartTM, [&](const auto& Downcast, const FTransform& StartFullTM)
 			{
-				return SweepHelper(Downcast, SpatialAcceleration, StartFullTM, Dir, DeltaMagnitude, HitBuffer, OutputFlags, QueryFilterData, QueryCallback, DebugParams);
+				return SweepHelper(Downcast, SpatialAcceleration, StartFullTM, Dir, DeltaMagnitude, HitBuffer, OutputFlags, CommonParams);
 			});
 		}
 
 		template<typename THit>
-		void Overlap(const Chaos::FImplicitObject& QueryGeom, const FTransform& GeomPose, ChaosInterface::FSQHitBuffer<THit>& HitBuffer, const ChaosInterface::FQueryFilterData& QueryFilterData, ICollisionQueryFilterCallbackBase& QueryCallback, const ChaosInterface::FQueryDebugParams& DebugParams = {}) const
+		void Overlap(const Chaos::FImplicitObject& QueryGeom, const FTransform& GeomPose, ChaosInterface::FSQHitBuffer<THit>& HitBuffer, const ChaosInterface::FSceneQueryCommonParams& CommonParams) const
 		{
 			return Chaos::Utilities::CastHelper(QueryGeom, GeomPose, [&](const auto& Downcast, const FTransform& GeomFullPose)
 			{
-				return OverlapHelper(Downcast, SpatialAcceleration, GeomFullPose, HitBuffer, QueryFilterData, QueryCallback, DebugParams);
+				return OverlapHelper(Downcast, SpatialAcceleration, GeomFullPose, HitBuffer, CommonParams);
 			});
 		}
 
@@ -90,7 +90,7 @@
 }
 
 template <typename TAccelContainer, typename THitRaycast>
-void LowLevelRaycastImpl(const TAccelContainer& Container, const FVector& Start, const FVector& Dir, float DeltaMag, FPhysicsHitCallback<THitRaycast>& HitBuffer, EHitFlags OutputFlags, FQueryFlags QueryFlags, const FCollisionFilterData& Filter, const ChaosInterface::FQueryFilterData& QueryFilterData, ICollisionQueryFilterCallbackBase* QueryCallback, const ChaosInterface::FQueryDebugParams& DebugParams)
+void LowLevelRaycastImpl(const TAccelContainer& Container, const FVector& Start, const FVector& Dir, float DeltaMag, FPhysicsHitCallback<THitRaycast>& HitBuffer, EHitFlags OutputFlags, const ChaosInterface::FSceneQueryCommonParams& CommonParams)
 {
 	constexpr bool bGTData = std::is_same_v<THitRaycast, FHitRaycast>;
 	using TTraits = FAccelerationContainerTraits<TAccelContainer, bGTData>;
@@ -100,22 +100,13 @@
 		double Time = 0.0;
 		{
 			FScopedDurationTimer Timer(Time);
-			SQAccelerator.Raycast(Start, Dir, DeltaMag, HitBuffer, OutputFlags, QueryFilterData, *QueryCallback, DebugParams);
+			SQAccelerator.Raycast(Start, Dir, DeltaMag, HitBuffer, OutputFlags, CommonParams);
 		}
 	}
 }
 
-#define DEFINE_LOW_LEVEL_RAYCAST_ACCEL_HIT(TACCEL, THIT) template<> void Chaos::Private::LowLevelRaycast(const TACCEL& Container, const FVector& Start, const FVector& Dir, float DeltaMag, FPhysicsHitCallback<THIT>& HitBuffer, EHitFlags OutputFlags, FQueryFlags QueryFlags, const FCollisionFilterData& Filter, const ChaosInterface::FQueryFilterData& QueryFilterData, ICollisionQueryFilterCallbackBase* QueryCallback, const ChaosInterface::FQueryDebugParams& DebugParams) \
-{ LowLevelRaycastImpl(Container, Start, Dir, DeltaMag, HitBuffer, OutputFlags, QueryFlags, Filter, QueryFilterData, QueryCallback, DebugParams); }
-#define DEFINE_LOW_LEVEL_RAYCAST_ACCEL(TACCEL) \
-DEFINE_LOW_LEVEL_RAYCAST_ACCEL_HIT(TACCEL, FHitRaycast); \
-DEFINE_LOW_LEVEL_RAYCAST_ACCEL_HIT(TACCEL, ChaosInterface::FPTRaycastHit);
-DEFINE_LOW_LEVEL_RAYCAST_ACCEL(Chaos::IDefaultChaosSpatialAcceleration)
-DEFINE_LOW_LEVEL_RAYCAST_ACCEL(IExternalSpatialAcceleration)
-DEFINE_LOW_LEVEL_RAYCAST_ACCEL(FPhysScene)
-
 template <typename TAccelContainer, typename THitSweep>
-void LowLevelSweepImpl(const TAccelContainer& Container, const FPhysicsGeometry& QueryGeom, const FTransform& StartTM, const FVector& Dir, float DeltaMag, FPhysicsHitCallback<THitSweep>& HitBuffer, EHitFlags OutputFlags, FQueryFlags QueryFlags, const FCollisionFilterData& Filter, const ChaosInterface::FQueryFilterData& QueryFilterData, ICollisionQueryFilterCallbackBase* QueryCallback, const ChaosInterface::FQueryDebugParams& DebugParams)
+void LowLevelSweepImpl(const TAccelContainer& Container, const FPhysicsGeometry& QueryGeom, const FTransform& StartTM, const FVector& Dir, float DeltaMag, FPhysicsHitCallback<THitSweep>& HitBuffer, EHitFlags OutputFlags, const ChaosInterface::FSceneQueryCommonParams& CommonParams)
 {
 	constexpr bool bGTData = std::is_same<THitSweep, FHitSweep>::value;
 	using TTraits = FAccelerationContainerTraits<TAccelContainer, bGTData>;
@@ -126,24 +117,14 @@
 			double Time = 0.0;
 			{
 				FScopedDurationTimer Timer(Time);
-				SQAccelerator.Sweep(QueryGeom, StartTM, Dir, DeltaMag, HitBuffer, OutputFlags, QueryFilterData, *QueryCallback, DebugParams);
+				SQAccelerator.Sweep(QueryGeom, StartTM, Dir, DeltaMag, HitBuffer, OutputFlags, CommonParams);
 			}
 		}
 	}
 }
 
-#define DEFINE_LOW_LEVEL_SWEEP_ACCEL_HIT(TACCEL, THIT) template<> void Chaos::Private::LowLevelSweep(const TACCEL& Container, const FPhysicsGeometry& QueryGeom, const FTransform& StartTM, const FVector& Dir, float DeltaMag, FPhysicsHitCallback<THIT>& HitBuffer, EHitFlags OutputFlags, FQueryFlags QueryFlags, const FCollisionFilterData& Filter, const ChaosInterface::FQueryFilterData& QueryFilterData, ICollisionQueryFilterCallbackBase* QueryCallback, const ChaosInterface::FQueryDebugParams& DebugParams) \
-{ LowLevelSweepImpl(Container, QueryGeom, StartTM, Dir, DeltaMag, HitBuffer, OutputFlags, QueryFlags, Filter, QueryFilterData, QueryCallback, DebugParams); }
-#define DEFINE_LOW_LEVEL_SWEEP_ACCEL(TACCEL) \
-DEFINE_LOW_LEVEL_SWEEP_ACCEL_HIT(TACCEL, FHitSweep); \
-DEFINE_LOW_LEVEL_SWEEP_ACCEL_HIT(TACCEL, ChaosInterface::FPTSweepHit);
-DEFINE_LOW_LEVEL_SWEEP_ACCEL(Chaos::IDefaultChaosSpatialAcceleration)
-DEFINE_LOW_LEVEL_SWEEP_ACCEL(IExternalSpatialAcceleration)
-DEFINE_LOW_LEVEL_SWEEP_ACCEL(FPhysScene)
-
-
 template <typename TAccelContainer, typename THitOverlap>
-void LowLevelOverlapImpl(const TAccelContainer& Container, const FPhysicsGeometry& QueryGeom, const FTransform& GeomPose, FPhysicsHitCallback<THitOverlap>& HitBuffer, FQueryFlags QueryFlags, const FCollisionFilterData& Filter, const ChaosInterface::FQueryFilterData& QueryFilterData, ICollisionQueryFilterCallbackBase* QueryCallback, const ChaosInterface::FQueryDebugParams& DebugParams)
+void LowLevelOverlapImpl(const TAccelContainer& Container, const FPhysicsGeometry& QueryGeom, const FTransform& GeomPose, FPhysicsHitCallback<THitOverlap>& HitBuffer, const ChaosInterface::FSceneQueryCommonParams& CommonParams)
 {
 	constexpr bool bGTData = std::is_same<THitOverlap, FHitOverlap>::value;
 	using TTraits = FAccelerationContainerTraits<TAccelContainer, bGTData>;
@@ -153,11 +134,86 @@
 		double Time = 0.0;
 		{
 			FScopedDurationTimer Timer(Time);
-			SQAccelerator.Overlap(QueryGeom, GeomPose, HitBuffer, QueryFilterData, *QueryCallback, DebugParams);
+			SQAccelerator.Overlap(QueryGeom, GeomPose, HitBuffer, CommonParams);
 		}
 	}
 }
 
+#define DEFINE_LOW_LEVEL_RAYCAST_ACCEL_HIT(TACCEL, THIT) template<> void Chaos::Private::LowLevelRaycast(const TACCEL& Container, const FVector& Start, const FVector& Dir, float DeltaMag, FPhysicsHitCallback<THIT>& HitBuffer, EHitFlags OutputFlags, const ChaosInterface::FSceneQueryCommonParams& CommonParams) \
+{ LowLevelRaycastImpl(Container, Start, Dir, DeltaMag, HitBuffer, OutputFlags, CommonParams); }
+#define DEFINE_LOW_LEVEL_RAYCAST_ACCEL(TACCEL) \
+DEFINE_LOW_LEVEL_RAYCAST_ACCEL_HIT(TACCEL, FHitRaycast); \
+DEFINE_LOW_LEVEL_RAYCAST_ACCEL_HIT(TACCEL, ChaosInterface::FPTRaycastHit);
+DEFINE_LOW_LEVEL_RAYCAST_ACCEL(Chaos::IDefaultChaosSpatialAcceleration)
+DEFINE_LOW_LEVEL_RAYCAST_ACCEL(IExternalSpatialAcceleration)
+DEFINE_LOW_LEVEL_RAYCAST_ACCEL(FPhysScene)
+#undef DEFINE_LOW_LEVEL_RAYCAST_ACCEL
+#undef DEFINE_LOW_LEVEL_RAYCAST_ACCEL_HIT
+
+#define DEFINE_LOW_LEVEL_SWEEP_ACCEL_HIT(TACCEL, THIT) template<> void Chaos::Private::LowLevelSweep(const TACCEL& Container, const FPhysicsGeometry& QueryGeom, const FTransform& StartTM, const FVector& Dir, float DeltaMag, FPhysicsHitCallback<THIT>& HitBuffer, EHitFlags OutputFlags, const ChaosInterface::FSceneQueryCommonParams& CommonParams) \
+{ LowLevelSweepImpl(Container, QueryGeom, StartTM, Dir, DeltaMag, HitBuffer, OutputFlags, CommonParams); }
+#define DEFINE_LOW_LEVEL_SWEEP_ACCEL(TACCEL) \
+DEFINE_LOW_LEVEL_SWEEP_ACCEL_HIT(TACCEL, FHitSweep); \
+DEFINE_LOW_LEVEL_SWEEP_ACCEL_HIT(TACCEL, ChaosInterface::FPTSweepHit);
+DEFINE_LOW_LEVEL_SWEEP_ACCEL(Chaos::IDefaultChaosSpatialAcceleration)
+DEFINE_LOW_LEVEL_SWEEP_ACCEL(IExternalSpatialAcceleration)
+DEFINE_LOW_LEVEL_SWEEP_ACCEL(FPhysScene)
+#undef DEFINE_LOW_LEVEL_SWEEP_ACCEL
+#undef DEFINE_LOW_LEVEL_SWEEP_ACCEL_HIT
+
+#define DEFINE_LOW_LEVEL_OVERLAP_ACCEL_HIT(TACCEL, THIT) template<> void Chaos::Private::LowLevelOverlap(const TACCEL& Container, const FPhysicsGeometry& QueryGeom, const FTransform& GeomPose, FPhysicsHitCallback<THIT>& HitBuffer, const ChaosInterface::FSceneQueryCommonParams& CommonParams) \
+{ LowLevelOverlapImpl(Container, QueryGeom, GeomPose, HitBuffer, CommonParams); }
+#define DEFINE_LOW_LEVEL_OVERLAP_ACCEL(TACCEL) \
+DEFINE_LOW_LEVEL_OVERLAP_ACCEL_HIT(TACCEL, FHitOverlap) \
+DEFINE_LOW_LEVEL_OVERLAP_ACCEL_HIT(TACCEL, ChaosInterface::FPTOverlapHit)
+DEFINE_LOW_LEVEL_OVERLAP_ACCEL(Chaos::IDefaultChaosSpatialAcceleration)
+DEFINE_LOW_LEVEL_OVERLAP_ACCEL(IExternalSpatialAcceleration)
+DEFINE_LOW_LEVEL_OVERLAP_ACCEL(FPhysScene)
+#undef DEFINE_LOW_LEVEL_OVERLAP_ACCEL
+#undef DEFINE_LOW_LEVEL_OVERLAP_ACCEL_HIT
+
+
+
+PRAGMA_DISABLE_DEPRECATION_WARNINGS
+template <typename TAccelContainer, typename THitRaycast>
+void LowLevelRaycastImpl(const TAccelContainer& Container, const FVector& Start, const FVector& Dir, float DeltaMag, FPhysicsHitCallback<THitRaycast>& HitBuffer, EHitFlags OutputFlags, FQueryFlags QueryFlags, const FCollisionFilterData& Filter, const ChaosInterface::FQueryFilterData& QueryFilterData, ICollisionQueryFilterCallbackBase* QueryCallback, const ChaosInterface::FQueryDebugParams& DebugParams)
+{
+	const ChaosInterface::FSceneQueryCommonParams CommonParams(*QueryCallback, QueryFilterData, DebugParams);
+	LowLevelRaycastImpl(Container, Start, Dir, DeltaMag, HitBuffer, OutputFlags, CommonParams);
+}
+
+#define DEFINE_LOW_LEVEL_RAYCAST_ACCEL_HIT(TACCEL, THIT) template<> void Chaos::Private::LowLevelRaycast(const TACCEL& Container, const FVector& Start, const FVector& Dir, float DeltaMag, FPhysicsHitCallback<THIT>& HitBuffer, EHitFlags OutputFlags, FQueryFlags QueryFlags, const FCollisionFilterData& Filter, const ChaosInterface::FQueryFilterData& QueryFilterData, ICollisionQueryFilterCallbackBase* QueryCallback, const ChaosInterface::FQueryDebugParams& DebugParams) \
+{ LowLevelRaycastImpl(Container, Start, Dir, DeltaMag, HitBuffer, OutputFlags, QueryFlags, Filter, QueryFilterData, QueryCallback, DebugParams); }
+#define DEFINE_LOW_LEVEL_RAYCAST_ACCEL(TACCEL) \
+DEFINE_LOW_LEVEL_RAYCAST_ACCEL_HIT(TACCEL, FHitRaycast); \
+DEFINE_LOW_LEVEL_RAYCAST_ACCEL_HIT(TACCEL, ChaosInterface::FPTRaycastHit);
+DEFINE_LOW_LEVEL_RAYCAST_ACCEL(Chaos::IDefaultChaosSpatialAcceleration)
+DEFINE_LOW_LEVEL_RAYCAST_ACCEL(IExternalSpatialAcceleration)
+DEFINE_LOW_LEVEL_RAYCAST_ACCEL(FPhysScene)
+
+template <typename TAccelContainer, typename THitSweep>
+void LowLevelSweepImpl(const TAccelContainer& Container, const FPhysicsGeometry& QueryGeom, const FTransform& StartTM, const FVector& Dir, float DeltaMag, FPhysicsHitCallback<THitSweep>& HitBuffer, EHitFlags OutputFlags, FQueryFlags QueryFlags, const FCollisionFilterData& Filter, const ChaosInterface::FQueryFilterData& QueryFilterData, ICollisionQueryFilterCallbackBase* QueryCallback, const ChaosInterface::FQueryDebugParams& DebugParams)
+{
+	const ChaosInterface::FSceneQueryCommonParams CommonParams(*QueryCallback, QueryFilterData, DebugParams);
+	LowLevelSweepImpl(Container, QueryGeom, StartTM, Dir, DeltaMag, HitBuffer, OutputFlags, CommonParams);
+}
+
+#define DEFINE_LOW_LEVEL_SWEEP_ACCEL_HIT(TACCEL, THIT) template<> void Chaos::Private::LowLevelSweep(const TACCEL& Container, const FPhysicsGeometry& QueryGeom, const FTransform& StartTM, const FVector& Dir, float DeltaMag, FPhysicsHitCallback<THIT>& HitBuffer, EHitFlags OutputFlags, FQueryFlags QueryFlags, const FCollisionFilterData& Filter, const ChaosInterface::FQueryFilterData& QueryFilterData, ICollisionQueryFilterCallbackBase* QueryCallback, const ChaosInterface::FQueryDebugParams& DebugParams) \
+{ LowLevelSweepImpl(Container, QueryGeom, StartTM, Dir, DeltaMag, HitBuffer, OutputFlags, QueryFlags, Filter, QueryFilterData, QueryCallback, DebugParams); }
+#define DEFINE_LOW_LEVEL_SWEEP_ACCEL(TACCEL) \
+DEFINE_LOW_LEVEL_SWEEP_ACCEL_HIT(TACCEL, FHitSweep); \
+DEFINE_LOW_LEVEL_SWEEP_ACCEL_HIT(TACCEL, ChaosInterface::FPTSweepHit);
+DEFINE_LOW_LEVEL_SWEEP_ACCEL(Chaos::IDefaultChaosSpatialAcceleration)
+DEFINE_LOW_LEVEL_SWEEP_ACCEL(IExternalSpatialAcceleration)
+DEFINE_LOW_LEVEL_SWEEP_ACCEL(FPhysScene)
+
+template <typename TAccelContainer, typename THitOverlap>
+void LowLevelOverlapImpl(const TAccelContainer& Container, const FPhysicsGeometry& QueryGeom, const FTransform& GeomPose, FPhysicsHitCallback<THitOverlap>& HitBuffer, FQueryFlags QueryFlags, const FCollisionFilterData& Filter, const ChaosInterface::FQueryFilterData& QueryFilterData, ICollisionQueryFilterCallbackBase* QueryCallback, const ChaosInterface::FQueryDebugParams& DebugParams)
+{
+	const ChaosInterface::FSceneQueryCommonParams CommonParams(*QueryCallback, QueryFilterData, DebugParams);
+	LowLevelOverlapImpl(Container, QueryGeom, GeomPose, HitBuffer, CommonParams);
+}
+
 #define DEFINE_LOW_LEVEL_OVERLAP_ACCEL_HIT(TACCEL, THIT) template<> void Chaos::Private::LowLevelOverlap(const TACCEL& Container, const FPhysicsGeometry& QueryGeom, const FTransform& GeomPose, FPhysicsHitCallback<THIT>& HitBuffer, FQueryFlags QueryFlags, const FCollisionFilterData& Filter, const ChaosInterface::FQueryFilterData& QueryFilterData, ICollisionQueryFilterCallbackBase* QueryCallback, const ChaosInterface::FQueryDebugParams& DebugParams) \
 { LowLevelOverlapImpl(Container, QueryGeom, GeomPose, HitBuffer, QueryFlags, Filter, QueryFilterData, QueryCallback, DebugParams); }
 #define DEFINE_LOW_LEVEL_OVERLAP_ACCEL(TACCEL) \
@@ -166,3 +222,4 @@
 DEFINE_LOW_LEVEL_OVERLAP_ACCEL(Chaos::IDefaultChaosSpatialAcceleration)
 DEFINE_LOW_LEVEL_OVERLAP_ACCEL(IExternalSpatialAcceleration)
 DEFINE_LOW_LEVEL_OVERLAP_ACCEL(FPhysScene)
+PRAGMA_ENABLE_DEPRECATION_WARNINGS
diff --git a/Engine/Source/Runtime/Engine/Private/PhysicsEngine/CollisionQueryFilterCallback.cpp b/Engine/Source/Runtime/Engine/Private/PhysicsEngine/CollisionQueryFilterCallback.cpp
--- a/Engine/Source/Runtime/Engine/Private/PhysicsEngine/CollisionQueryFilterCallback.cpp
+++ b/Engine/Source/Runtime/Engine/Private/PhysicsEngine/CollisionQueryFilterCallback.cpp
@@ -81,8 +81,28 @@ ECollisionQueryHitType FCollisionQueryFilterCallback::CalcQueryHitType(const FCo
 	return ECollisionQueryHitType::None;
 }
 
+ECollisionQueryHitType FCollisionQueryFilterCallback::PostFilter(const FQueryFilterData& FilterData, const ChaosInterface::FQueryHit& Hit)
+{
+	return PostFilterImp(FilterData, Hit);
+}
+
+ECollisionQueryHitType FCollisionQueryFilterCallback::PreFilter(const FQueryFilterData& FilterData, const Chaos::FPerShapeData& Shape, const Chaos::FGeometryParticle& Actor)
+{
+	return PreFilterImp(FilterData, Shape, Actor);
+}
+
+ECollisionQueryHitType FCollisionQueryFilterCallback::PostFilter(const FQueryFilterData& FilterData, const ChaosInterface::FPTQueryHit& Hit)
+{
+	return PostFilterImp(FilterData, Hit);
+}
+
+ECollisionQueryHitType FCollisionQueryFilterCallback::PreFilter(const FQueryFilterData& FilterData, const Chaos::FPerShapeData& Shape, const Chaos::FGeometryParticleHandle& Actor)
+{
+	return PreFilterImp(FilterData, Shape, Actor);
+}
+
 template <typename TParticle>
-ECollisionQueryHitType FCollisionQueryFilterCallback::PreFilterBaseImp(const FCollisionFilterData& FilterData, const Chaos::FPerShapeData& Shape, const TParticle& Actor)
+ECollisionQueryHitType FCollisionQueryFilterCallback::PreFilterBaseImp(const FQueryFilterData& FilterData, const Chaos::FPerShapeData& Shape, const TParticle& Actor)
 {
 	//SCOPE_CYCLE_COUNTER(STAT_Collision_PreFilter);
 
@@ -91,16 +111,7 @@ ECollisionQueryHitType FCollisionQueryFilterCallback::PreFilterBaseImp(const FCo
 		return ECollisionQueryHitType::None;
 	}
 
-	const Chaos::Filter::FCombinedShapeFilterData CombinedShapeFilter = ChaosInterface::GetCombinedShapeFilterData(Shape);
-	const Chaos::Filter::FInstanceData& InstanceData = CombinedShapeFilter.GetInstanceData();
-	FCollisionFilterData ShapeFilter = Chaos::Filter::FShapeFilterBuilder::GetLegacyShapeQueryFilter(CombinedShapeFilter);
-
-	// We usually don't have ignore components so we try to avoid the virtual getSimulationFilterData() call below. 'word2' of shape sim filter data is componentID.
-	uint32 ComponentID = 0;
-	if (IgnoreComponents.Num() > 0)
-	{
-		ComponentID = InstanceData.GetComponentId();
-	}
+	const FCombinedShapeFilterData CombinedShapeFilter = ChaosInterface::GetCombinedShapeFilterData(Shape);
 
 	FBodyInstance* BodyInstance = nullptr;
 
@@ -111,20 +122,20 @@ ECollisionQueryHitType FCollisionQueryFilterCallback::PreFilterBaseImp(const FCo
 #endif // ENABLE_PREFILTER_LOGGING || DETECT_SQ_HITCHES
 	}
 
-	return PreFilterImp(FilterData, ShapeFilter, ComponentID, BodyInstance);
+	return PreFilterImp(FilterData, CombinedShapeFilter, BodyInstance);
 }
 
-ECollisionQueryHitType FCollisionQueryFilterCallback::PreFilterImp(const FCollisionFilterData& FilterData, const Chaos::FPerShapeData& Shape, const Chaos::FGeometryParticle& Actor)
+ECollisionQueryHitType FCollisionQueryFilterCallback::PreFilterImp(const FQueryFilterData& FilterData, const Chaos::FPerShapeData& Shape, const Chaos::FGeometryParticle& Actor)
 {
 	return PreFilterBaseImp(FilterData, Shape, Actor);
 }
 
-ECollisionQueryHitType FCollisionQueryFilterCallback::PreFilterImp(const FCollisionFilterData& FilterData, const Chaos::FPerShapeData& Shape, const Chaos::FGeometryParticleHandle& Actor)
+ECollisionQueryHitType FCollisionQueryFilterCallback::PreFilterImp(const FQueryFilterData& FilterData, const Chaos::FPerShapeData& Shape, const Chaos::FGeometryParticleHandle& Actor)
 {
 	return PreFilterBaseImp(FilterData, Shape, Actor);
 }
 
-ECollisionQueryHitType FCollisionQueryFilterCallback::PreFilterImp(const FCollisionFilterData& FilterData, const FCollisionFilterData& ShapeFilter, uint32 ComponentID, const FBodyInstance* BodyInstance)
+ECollisionQueryHitType FCollisionQueryFilterCallback::PreFilterImp(const FQueryFilterData& FilterData, const FCombinedShapeFilterData& CombinedShapeFilter, const FBodyInstance* BodyInstance)
 {
 #if DETECT_SQ_HITCHES
 	FPreFilterRecord* PreFilterRecord = nullptr;
@@ -143,6 +154,9 @@ ECollisionQueryHitType FCollisionQueryFilterCallback::PreFilterImp(const FCollis
 	}
 #endif
 
+	const FInstanceData& ShapeInstanceData = CombinedShapeFilter.GetInstanceData();
+	const FShapeFilterData& ShapeFilter = CombinedShapeFilter.GetShapeFilterData();
+
 #if ENABLE_PREFILTER_LOGGING
 	static bool bLoggingEnabled = false;
 	if (bLoggingEnabled)
@@ -154,24 +168,27 @@ ECollisionQueryHitType FCollisionQueryFilterCallback::PreFilterImp(const FCollis
 				*BodyInstance->OwnerComponent.Get()->GetName());
 		}
 
-		UE_LOG(LogCollision, Warning, TEXT("ShapeFilter : %x %x %x %x"), ShapeFilter.Word0, ShapeFilter.Word1, ShapeFilter.Word2, ShapeFilter.Word3);
-		UE_LOG(LogCollision, Warning, TEXT("FilterData : %x %x %x %x"), FilterData.Word0, FilterData.Word1, FilterData.Word2, FilterData.Word3);
+		UE_LOG(LogCollision, Warning, TEXT("ShapeInstanceData : %s"), *ShapeInstanceData.ToString());
+		UE_LOG(LogCollision, Warning, TEXT("ShapeFilter : %s"), *ShapeFilter.ToString());
+		UE_LOG(LogCollision, Warning, TEXT("QueryFilter : %s"), *FilterData.ToString());
 	}
 #endif // ENABLE_PREFILTER_LOGGING
 
 	// Shape : shape's Filter Data
 	// Querier : filterData that owns the trace
-	uint32 ShapeFlags = ShapeFilter.Word3 & 0xFFFFFF;
-	uint32 QuerierFlags = FilterData.Word3 & 0xFFFFFF;
-	uint32 CommonFlags = ShapeFlags & QuerierFlags;
+	const Chaos::EFilterFlags ShapeFlags = ShapeFilter.GetFlags();
+	const Chaos::EFilterFlags QuerierFlags = FilterData.GetFlags();
+	const Chaos::EFilterFlags CommonFlags = ShapeFlags & QuerierFlags;
 
 	// First check complexity, none of them matches
-	if (!(CommonFlags & EPDF_SimpleCollision) && !(CommonFlags & EPDF_ComplexCollision))
+	if (!(CommonFlags & Chaos::EFilterFlags::SimpleCollision) && !(CommonFlags & Chaos::EFilterFlags::ComplexCollision))
 	{
 		return (PreFilterReturnValue = ECollisionQueryHitType::None);
 	}
-
-	ECollisionQueryHitType Result = FCollisionQueryFilterCallback::CalcQueryHitType(FilterData, ShapeFilter, true);
+	
+	const FCollisionFilterData LegacyQueryFilter = Chaos::Filter::FQueryFilterBuilder::GetLegacyQueryFilter(FilterData);
+	const FCollisionFilterData LegacyShapeFilter = Chaos::Filter::FShapeFilterBuilder::GetLegacyShapeQueryFilter(CombinedShapeFilter);
+	ECollisionQueryHitType Result = FCollisionQueryFilterCallback::CalcQueryHitType(LegacyQueryFilter, LegacyShapeFilter, true);
 
 	if (Result == ECollisionQueryHitType::Touch && bIgnoreTouches)
 	{
@@ -187,13 +204,13 @@ ECollisionQueryHitType FCollisionQueryFilterCallback::PreFilterImp(const FCollis
 	if (Result != ECollisionQueryHitType::None)
 	{
 		// See if we are ignoring the actor this shape belongs to (word0 of shape filterdata is actorID)
-		if (IgnoreActors.Contains(ShapeFilter.Word0))
+		if (IgnoreActors.Contains(ShapeInstanceData.GetOwnerId()))
 		{
 			//UE_LOG(LogTemp, Log, TEXT("Ignoring Actor: %d"), ShapeFilter.word0);
 			Result = ECollisionQueryHitType::None;
 		}
 
-		if (IgnoreComponents.Contains(ComponentID))
+		if (IgnoreComponents.Contains(ShapeInstanceData.GetComponentId()))
 		{
 			//UE_LOG(LogTemp, Log, TEXT("Ignoring Component: %d"), shape->getSimulationFilterData().word2);
 			Result = ECollisionQueryHitType::None;
@@ -203,8 +220,8 @@ ECollisionQueryHitType FCollisionQueryFilterCallback::PreFilterImp(const FCollis
 #if ENABLE_PREFILTER_LOGGING
 	if (bLoggingEnabled)
 	{
-		ECollisionChannel QuerierChannel = GetCollisionChannel(FilterData.word3);
-		UE_LOG(LogCollision, Log, TEXT("[PREFILTER] Result for Querier [CHANNEL: %d, FLAG: %x] %s [%d]"), (int32)QuerierChannel, QuerierFlags, *ComponentName, (int32)Result);
+		uint32 QuerierChannel = FilterData.GetCollisionChannelIndex();
+		UE_LOG(LogCollision, Log, TEXT("[PREFILTER] Result for Querier [CHANNEL: %d, FLAG: %x] [%d]"), QuerierChannel, QuerierFlags, (int32)Result);
 	}
 #endif // ENABLE_PREFILTER_LOGGING
 
@@ -224,7 +241,7 @@ ECollisionQueryHitType FCollisionQueryFilterCallback::PreFilterImp(const FCollis
 }
 
 
-ECollisionQueryHitType FCollisionQueryFilterCallback::PostFilterImp(const FCollisionFilterData& FilterData, bool bIsOverlap)
+ECollisionQueryHitType FCollisionQueryFilterCallback::PostFilterImp(const FQueryFilterData& FilterData, bool bIsOverlap)
 {
 	if (!bIsSweep)
 	{
@@ -247,7 +264,7 @@ ECollisionQueryHitType FCollisionQueryFilterCallback::PostFilterImp(const FColli
 	}
 }
 
-ECollisionQueryHitType FCollisionQueryFilterCallback::PostFilterImp(const FCollisionFilterData& FilterData, const ChaosInterface::FQueryHit& Hit)
+ECollisionQueryHitType FCollisionQueryFilterCallback::PostFilterImp(const FQueryFilterData& FilterData, const ChaosInterface::FQueryHit& Hit)
 {
 	// Unused in non-sweeps
 	if (!bIsSweep)
@@ -261,7 +278,7 @@ ECollisionQueryHitType FCollisionQueryFilterCallback::PostFilterImp(const FColli
 	return PostFilterImp(FilterData, bIsOverlap);
 }
 
-ECollisionQueryHitType FCollisionQueryFilterCallback::PostFilterImp(const FCollisionFilterData& FilterData, const ChaosInterface::FPTQueryHit& Hit)
+ECollisionQueryHitType FCollisionQueryFilterCallback::PostFilterImp(const FQueryFilterData& FilterData, const ChaosInterface::FPTQueryHit& Hit)
 {
 	// Unused in non-sweeps
 	if (!bIsSweep)
diff --git a/Engine/Source/Runtime/Engine/Private/PhysicsEngine/CollisionQueryFilterCallback.h b/Engine/Source/Runtime/Engine/Private/PhysicsEngine/CollisionQueryFilterCallback.h
--- a/Engine/Source/Runtime/Engine/Private/PhysicsEngine/CollisionQueryFilterCallback.h
+++ b/Engine/Source/Runtime/Engine/Private/PhysicsEngine/CollisionQueryFilterCallback.h
@@ -16,9 +16,14 @@ typedef FCollisionQueryParams::IgnoreComponentsArrayType FilterIgnoreComponentsA
 /** TArray typedef of actors to ignore. */
 typedef FCollisionQueryParams::IgnoreActorsArrayType FilterIgnoreActorsArrayType;
 
-class FCollisionQueryFilterCallback : public ICollisionQueryFilterCallbackBase
+class FCollisionQueryFilterCallback : public ICollisionQueryFilterCallback
 {
 public:
+	using FInstanceData = Chaos::Filter::FInstanceData;
+	using FQueryFilterData = Chaos::Filter::FQueryFilterData;
+	using FShapeFilterData = Chaos::Filter::FShapeFilterData;
+	using FCombinedShapeFilterData = Chaos::Filter::FCombinedShapeFilterData;
+
 	/** Result of PreFilter callback. */
 	ECollisionQueryHitType PreFilterReturnValue;
 
@@ -55,38 +60,13 @@ public:
 	}
 	~FCollisionQueryFilterCallback() {}
 
-
 	static ECollisionQueryHitType CalcQueryHitType(const FCollisionFilterData& QueryFilter, const FCollisionFilterData& ShapeFilter, bool bPreFilter = false);
 
-	ECollisionQueryHitType PreFilterImp(const FCollisionFilterData& FilterData, const FCollisionFilterData& ShapeFilterData, uint32 ComponentID, const FBodyInstance* BodyInstance);
-
-	ECollisionQueryHitType PostFilterImp(const FCollisionFilterData& FilterData, bool bIsOverlap);
-
-	ECollisionQueryHitType PreFilterImp(const FCollisionFilterData& FilterData, const Chaos::FPerShapeData& Shape, const Chaos::FGeometryParticle& Actor);
-	ECollisionQueryHitType PreFilterImp(const FCollisionFilterData& FilterData, const Chaos::FPerShapeData& Shape, const Chaos::FGeometryParticleHandle& Actor);
-
-	ECollisionQueryHitType PostFilterImp(const FCollisionFilterData& FilterData, const ChaosInterface::FQueryHit& Hit);
-	ECollisionQueryHitType PostFilterImp(const FCollisionFilterData& FilterData, const ChaosInterface::FPTQueryHit& Hit);
-
-	virtual ECollisionQueryHitType PostFilter(const FCollisionFilterData& FilterData, const ChaosInterface::FQueryHit& Hit) override
-	{
-		return PostFilterImp(FilterData, Hit);
-	}
-
-	virtual ECollisionQueryHitType PostFilter(const FCollisionFilterData& FilterData, const ChaosInterface::FPTQueryHit& Hit) override
-	{
-		return PostFilterImp(FilterData, Hit);
-	}
+	virtual ECollisionQueryHitType PostFilter(const FQueryFilterData& FilterData, const ChaosInterface::FQueryHit& Hit) override;
+	virtual ECollisionQueryHitType PostFilter(const FQueryFilterData& FilterData, const ChaosInterface::FPTQueryHit& Hit) override;
 
-	virtual ECollisionQueryHitType PreFilter(const FCollisionFilterData& FilterData, const Chaos::FPerShapeData& Shape, const Chaos::FGeometryParticle& Actor) override
-	{
-		return PreFilterImp(FilterData, Shape, Actor);
-	}
-
-	virtual ECollisionQueryHitType PreFilter(const FCollisionFilterData& FilterData, const Chaos::FPerShapeData& Shape, const Chaos::FGeometryParticleHandle& Actor) override
-	{
-		return PreFilterImp(FilterData, Shape, Actor);
-	}
+	virtual ECollisionQueryHitType PreFilter(const FQueryFilterData& FilterData, const Chaos::FPerShapeData& Shape, const Chaos::FGeometryParticle& Actor) override;
+	virtual ECollisionQueryHitType PreFilter(const FQueryFilterData& FilterData, const Chaos::FPerShapeData& Shape, const Chaos::FGeometryParticleHandle& Actor) override;
 
 #if DETECT_SQ_HITCHES
 	// Util struct to record what preFilter was called with
@@ -104,6 +84,14 @@ public:
 
 private:
 	template <typename TParticle>
-	ECollisionQueryHitType PreFilterBaseImp(const FCollisionFilterData& FilterData, const Chaos::FPerShapeData& Shape, const TParticle& Actor);
+	ECollisionQueryHitType PreFilterBaseImp(const FQueryFilterData& FilterData, const Chaos::FPerShapeData& Shape, const TParticle& Actor);
+
+	ECollisionQueryHitType PreFilterImp(const FQueryFilterData& FilterData, const Chaos::FPerShapeData& Shape, const Chaos::FGeometryParticle& Actor);
+	ECollisionQueryHitType PreFilterImp(const FQueryFilterData& FilterData, const Chaos::FPerShapeData& Shape, const Chaos::FGeometryParticleHandle& Actor);
+	ECollisionQueryHitType PreFilterImp(const FQueryFilterData& FilterData, const FCombinedShapeFilterData& CombinedShapeFilterData, const FBodyInstance* BodyInstance);
+
+	ECollisionQueryHitType PostFilterImp(const FQueryFilterData& FilterData, bool bIsOverlap);
+	ECollisionQueryHitType PostFilterImp(const FQueryFilterData& FilterData, const ChaosInterface::FQueryHit& Hit);
+	ECollisionQueryHitType PostFilterImp(const FQueryFilterData& FilterData, const ChaosInterface::FPTQueryHit& Hit);
 };
 
diff --git a/Engine/Source/Runtime/Engine/Public/Physics/Experimental/ChaosInterfaceWrapper.h b/Engine/Source/Runtime/Engine/Public/Physics/Experimental/ChaosInterfaceWrapper.h
--- a/Engine/Source/Runtime/Engine/Public/Physics/Experimental/ChaosInterfaceWrapper.h
+++ b/Engine/Source/Runtime/Engine/Public/Physics/Experimental/ChaosInterfaceWrapper.h
@@ -7,6 +7,7 @@
 #include "PhysicsInterfaceUtilsCore.h"
 #include "CollisionQueryFilterCallbackCore.h"
 #include "Chaos/PBDRigidsEvolutionFwd.h"
+#include "SceneQueryCommonParams.h"
 
 class FPhysScene_Chaos;
 class UPhysicalMaterial;
@@ -34,12 +35,13 @@ struct FScopedSceneReadLock
 	Chaos::FPBDRigidsSolver* Solver;
 };
 
+PRAGMA_DISABLE_DEPRECATION_WARNINGS
+UE_DEPRECATED(5.8, "Use ChaosInterface::FSceneQueryCommonParams constructor instead.")
 inline FQueryFilterData MakeQueryFilterData(const FCollisionFilterData& FilterData, EQueryFlags QueryFlags, const FCollisionQueryParams& Params)
 {
-	PRAGMA_DISABLE_INTERNAL_WARNINGS
 	return FChaosQueryFilterData(U2CFilterData(FilterData), U2CQueryFlags(QueryFlags));
-	PRAGMA_ENABLE_INTERNAL_WARNINGS
 }
+PRAGMA_ENABLE_DEPRECATION_WARNINGS
 
 FBodyInstance* GetUserData(const Chaos::FGeometryParticle& Actor);
 UPhysicalMaterial* GetUserData(const Chaos::FChaosPhysicsMaterial& Material);
@@ -51,12 +53,61 @@ bool HasValidUserData(const Chaos::FGeometryParticle& Actor);
 namespace Chaos::Private
 {
 	template<typename TContainer, typename THit>
+	UE_INTERNAL void LowLevelRaycast(const TContainer& Container, const FVector& Start, const FVector& Dir, float DeltaMag, FPhysicsHitCallback<THit>& HitBuffer, EHitFlags OutputFlags, const ChaosInterface::FSceneQueryCommonParams& CommonParams);
+
+	template<typename TContainer, typename THit>
+	UE_INTERNAL void LowLevelSweep(const TContainer& Container, const FPhysicsGeometry& Geom, const FTransform& StartTM, const FVector& Dir, float DeltaMag, FPhysicsHitCallback<THit>& HitBuffer, EHitFlags OutputFlags, const ChaosInterface::FSceneQueryCommonParams& CommonParams);
+
+	template<typename TContainer, typename THit>
+	UE_INTERNAL void LowLevelOverlap(const TContainer& Container, const FPhysicsGeometry& Geom, const FTransform& GeomPose, FPhysicsHitCallback<THit>& HitBuffer, const ChaosInterface::FSceneQueryCommonParams& CommonParams);
+
+	// Export raycast functions
+#define DECLARE_LOW_LEVEL_RAYCAST_ACCEL_HIT(TACCEL, THIT) template<> UE_INTERNAL ENGINE_API void LowLevelRaycast(const TACCEL& Container, const FVector& Start, const FVector& Dir, float DeltaMag, FPhysicsHitCallback<THIT>& HitBuffer, EHitFlags OutputFlags, const ChaosInterface::FSceneQueryCommonParams& CommonParams);
+#define DECLARE_LOW_LEVEL_RAYCAST_ACCEL(TACCEL) \
+DECLARE_LOW_LEVEL_RAYCAST_ACCEL_HIT(TACCEL, FHitRaycast); \
+DECLARE_LOW_LEVEL_RAYCAST_ACCEL_HIT(TACCEL, ChaosInterface::FPTRaycastHit);
+DECLARE_LOW_LEVEL_RAYCAST_ACCEL(Chaos::IDefaultChaosSpatialAcceleration)
+DECLARE_LOW_LEVEL_RAYCAST_ACCEL(IExternalSpatialAcceleration)
+DECLARE_LOW_LEVEL_RAYCAST_ACCEL(FPhysScene)
+#undef DECLARE_LOW_LEVEL_RAYCAST_ACCEL
+#undef DECLARE_LOW_LEVEL_RAYCAST_ACCEL_HIT
+
+// Export sweep functions
+#define DECLARE_LOW_LEVEL_SWEEP_ACCEL_HIT(TACCEL, THIT) template<> UE_INTERNAL ENGINE_API void LowLevelSweep(const TACCEL& Container, const FPhysicsGeometry& QueryGeom, const FTransform& StartTM, const FVector& Dir, float DeltaMag, FPhysicsHitCallback<THIT>& HitBuffer, EHitFlags OutputFlags, const ChaosInterface::FSceneQueryCommonParams& CommonParams);
+#define DECLARE_LOW_LEVEL_SWEEP_ACCEL(TACCEL) \
+DECLARE_LOW_LEVEL_SWEEP_ACCEL_HIT(TACCEL, FHitSweep); \
+DECLARE_LOW_LEVEL_SWEEP_ACCEL_HIT(TACCEL, ChaosInterface::FPTSweepHit);
+DECLARE_LOW_LEVEL_SWEEP_ACCEL(Chaos::IDefaultChaosSpatialAcceleration)
+DECLARE_LOW_LEVEL_SWEEP_ACCEL(IExternalSpatialAcceleration)
+DECLARE_LOW_LEVEL_SWEEP_ACCEL(FPhysScene)
+#undef DECLARE_LOW_LEVEL_SWEEP_ACCEL
+#undef DECLARE_LOW_LEVEL_SWEEP_ACCEL_HIT
+
+// Export overlap functions
+#define DECLARE_LOW_LEVEL_OVERLAP_ACCEL_HIT(TACCEL, THIT) template<> UE_INTERNAL ENGINE_API void LowLevelOverlap(const TACCEL& Container, const FPhysicsGeometry& QueryGeom, const FTransform& GeomPose, FPhysicsHitCallback<THIT>& HitBuffer, const ChaosInterface::FSceneQueryCommonParams& CommonParams)
+#define DECLARE_LOW_LEVEL_OVERLAP_ACCEL(TACCEL) \
+DECLARE_LOW_LEVEL_OVERLAP_ACCEL_HIT(TACCEL, FHitOverlap); \
+DECLARE_LOW_LEVEL_OVERLAP_ACCEL_HIT(TACCEL, ChaosInterface::FPTOverlapHit);
+DECLARE_LOW_LEVEL_OVERLAP_ACCEL(Chaos::IDefaultChaosSpatialAcceleration)
+DECLARE_LOW_LEVEL_OVERLAP_ACCEL(IExternalSpatialAcceleration)
+DECLARE_LOW_LEVEL_OVERLAP_ACCEL(FPhysScene)
+#undef DECLARE_LOW_LEVEL_OVERLAP_ACCEL
+#undef DECLARE_LOW_LEVEL_OVERLAP_ACCEL_HIT
+}
+
+PRAGMA_DISABLE_DEPRECATION_WARNINGS
+namespace Chaos::Private
+{
+	template<typename TContainer, typename THit>
+	UE_DEPRECATED(5.8, "This is meant to be internal only")
 	void LowLevelRaycast(const TContainer& Container, const FVector& Start, const FVector& Dir, float DeltaMag, FPhysicsHitCallback<THit>& HitBuffer, EHitFlags OutputFlags, FQueryFlags QueryFlags, const FCollisionFilterData& Filter, const ChaosInterface::FQueryFilterData& QueryFilterData, ICollisionQueryFilterCallbackBase* QueryCallback, const ChaosInterface::FQueryDebugParams& DebugParams = {});
 
 	template<typename TContainer, typename THit>
+	UE_DEPRECATED(5.8, "This is meant to be internal only")
 	void LowLevelSweep(const TContainer& Container, const FPhysicsGeometry& Geom, const FTransform& StartTM, const FVector& Dir, float DeltaMag, FPhysicsHitCallback<THit>& HitBuffer, EHitFlags OutputFlags, FQueryFlags QueryFlags, const FCollisionFilterData& Filter, const ChaosInterface::FQueryFilterData& QueryFilterData, ICollisionQueryFilterCallbackBase* QueryCallback, const ChaosInterface::FQueryDebugParams& DebugParams = {});
 
 	template<typename TContainer, typename THit>
+	UE_DEPRECATED(5.8, "This is meant to be internal only")
 	void LowLevelOverlap(const TContainer& Container, const FPhysicsGeometry& Geom, const FTransform& GeomPose, FPhysicsHitCallback<THit>& HitBuffer, FQueryFlags QueryFlags, const FCollisionFilterData& Filter, const ChaosInterface::FQueryFilterData& QueryFilterData, ICollisionQueryFilterCallbackBase* QueryCallback, const ChaosInterface::FQueryDebugParams& DebugParams = {});
 }
 
@@ -86,3 +137,5 @@ DECLARE_LOW_LEVEL_OVERLAP_ACCEL_HIT(TACCEL, ChaosInterface::FPTOverlapHit);
 DECLARE_LOW_LEVEL_OVERLAP_ACCEL(Chaos::IDefaultChaosSpatialAcceleration)
 DECLARE_LOW_LEVEL_OVERLAP_ACCEL(IExternalSpatialAcceleration)
 DECLARE_LOW_LEVEL_OVERLAP_ACCEL(FPhysScene)
+
+PRAGMA_ENABLE_DEPRECATION_WARNINGS
diff --git a/Engine/Source/Runtime/Experimental/Chaos/Private/Chaos/CollisionFilterData.cpp b/Engine/Source/Runtime/Experimental/Chaos/Private/Chaos/CollisionFilterData.cpp
--- a/Engine/Source/Runtime/Experimental/Chaos/Private/Chaos/CollisionFilterData.cpp
+++ b/Engine/Source/Runtime/Experimental/Chaos/Private/Chaos/CollisionFilterData.cpp
@@ -131,6 +131,11 @@ namespace Chaos::Filter
 		ComponentId = InComponentId;
 	}
 
+	FString FInstanceData::ToString() const
+	{
+		return FString::Format(TEXT("OwnerId(%d) ComponentId(%d)"), { GetOwnerId(), GetComponentId() });
+	}
+
 	bool FShapeFilterData::IsValid() const
 	{
 		return QueryBlockChannels != 0 || QueryOverlapChannels != 0 || SimBlockChannels != 0 || Word3 != 0 || SimWord3 != 0;
@@ -188,6 +193,11 @@ namespace Chaos::Filter
 		return SimBlockChannels;
 	}
 
+	FString FShapeFilterData::ToString() const
+	{
+		return FString::Format(TEXT("SimBlock(%x) QueryBlock(%x) QueryOverlap(%x) CollisionChannel(%x) MaskFilter(%x) Flags(%x)"), { GetSimBlockChannels(), GetQueryBlockChannels(), GetQueryOverlapChannels(), GetCollisionChannelIndex(), GetMaskFilter(), (uint8)GetFlags() });
+	}
+
 	uint32 FShapeFilterData::GetLegacyWord3() const
 	{
 		return Word3;
@@ -254,6 +264,18 @@ namespace Chaos::Filter
 		return MultiTrace != 0;
 	}
 
+	CHAOS_API FString FQueryFilterData::ToString() const
+	{
+		if (GetQueryType() == EQueryType::Channel)
+		{
+			return FString::Format(TEXT("Type(Channel) Block(%x) Overlap(%x) ChannelIndex(%x) IgnoreMask(%d) Flags(%x)"), { GetBlockChannels(), GetOverlapChannels(), GetCollisionChannelIndex(), GetIgnoreMask(), (uint8)GetFlags() });
+		}
+		else
+		{
+			return FString::Format(TEXT("Type(Object) IsMulti(%d) ObjectTypes(%x) IgnoreMask(%d) Flags(%x)"), { IsMultiQuery(), GetObjectTypesToQueryMask(), GetIgnoreMask(), (uint8)GetFlags() });
+		}
+	}
+
 	void FQueryFilterData::SetMaskChannelAndFlags(const uint8 IgnoreMask, const uint8 ChannelIndex, const EFilterFlags FilterFlags)
 	{
 		Word3 = ((uint32)IgnoreMask << MaskFilterBitsOffset) | ((uint32)ChannelIndex << CollisionChannelBitsOffset) | (uint32)FilterFlags;
diff --git a/Engine/Source/Runtime/Experimental/Chaos/Public/Chaos/CollisionFilterData.h b/Engine/Source/Runtime/Experimental/Chaos/Public/Chaos/CollisionFilterData.h
--- a/Engine/Source/Runtime/Experimental/Chaos/Public/Chaos/CollisionFilterData.h
+++ b/Engine/Source/Runtime/Experimental/Chaos/Public/Chaos/CollisionFilterData.h
@@ -95,6 +95,8 @@ namespace Chaos::Filter
 		UE_INTERNAL CHAOS_API uint32 GetComponentId() const;
 		UE_INTERNAL CHAOS_API void SetComponentId(const uint32 InComponentId);
 
+		CHAOS_API FString ToString() const;
+
 	private:
 		friend class FShapeFilterBuilder;
 		friend FCombinedShapeFilterData;
@@ -128,6 +130,8 @@ namespace Chaos::Filter
 		CHAOS_API uint64 GetQueryOverlapChannels() const;
 		CHAOS_API uint64 GetSimBlockChannels() const;
 
+		CHAOS_API FString ToString() const;
+
 		UE_INTERNAL uint32 GetLegacyWord3() const;
 
 	private:
@@ -176,6 +180,8 @@ namespace Chaos::Filter
 		CHAOS_API uint64 GetObjectTypesToQueryMask() const;
 		CHAOS_API bool IsMultiQuery() const;
 
+		CHAOS_API FString ToString() const;
+
 	private:
 		friend class FQueryFilterBuilder;
 		friend class FQueryObjectFilterBuilder;
diff --git a/Engine/Source/Runtime/PhysicsCore/Private/PhysicsInterfaceUtilsCore.cpp b/Engine/Source/Runtime/PhysicsCore/Private/PhysicsInterfaceUtilsCore.cpp
--- a/Engine/Source/Runtime/PhysicsCore/Private/PhysicsInterfaceUtilsCore.cpp
+++ b/Engine/Source/Runtime/PhysicsCore/Private/PhysicsInterfaceUtilsCore.cpp
@@ -6,6 +6,7 @@
 
 #include "Chaos/ParticleHandle.h"
 
+PRAGMA_DISABLE_DEPRECATION_WARNINGS
 FCollisionFilterData C2UFilterData(const FChaosFilterData& PFilterData)
 {
 	FCollisionFilterData FilterData;
@@ -25,6 +26,7 @@ FCollisionFilterData ToUnrealFilterData(const FChaosFilterData& FilterData)
 {
 	return C2UFilterData(FilterData);
 }
+PRAGMA_ENABLE_DEPRECATION_WARNINGS
 
 template <typename THitLocation>
 uint32 FindFaceIndexImp(const THitLocation& PHit, const FVector& UnitDir)
diff --git a/Engine/Source/Runtime/PhysicsCore/Private/SQAccelerator.cpp b/Engine/Source/Runtime/PhysicsCore/Private/SQAccelerator.cpp
--- a/Engine/Source/Runtime/PhysicsCore/Private/SQAccelerator.cpp
+++ b/Engine/Source/Runtime/PhysicsCore/Private/SQAccelerator.cpp
@@ -62,25 +62,70 @@ struct FPreFilterInfo
 };
 
 template <typename TRaycastHit>
-void FChaosSQAccelerator::RaycastImp(const FVector& Start, const FVector& Dir, const float DeltaMagnitude, ChaosInterface::FSQHitBuffer<TRaycastHit>& HitBuffer, EHitFlags OutputFlags, const ChaosInterface::FQueryFilterData& QueryFilterData, ICollisionQueryFilterCallbackBase& QueryCallback, const ChaosInterface::FQueryDebugParams& DebugParams) const
+void FChaosSQAccelerator::RaycastImp(const FVector& Start, const FVector& Dir, const float DeltaMagnitude, ChaosInterface::FSQHitBuffer<TRaycastHit>& HitBuffer, EHitFlags OutputFlags, const ChaosInterface::FSceneQueryCommonParams& CommonParams) const
 {
 	using namespace Chaos;
 	using namespace ChaosInterface;
 
-	TSQVisitor<TSphere<FReal, 3>, FAccelerationStructureHandle, TRaycastHit, std::is_same<TRaycastHit, FRaycastHit>::value> RaycastVisitor(Start, Dir, HitBuffer, OutputFlags, QueryFilterData, QueryCallback, DebugParams);
+	TSQVisitor<TSphere<FReal, 3>, FAccelerationStructureHandle, TRaycastHit, std::is_same<TRaycastHit, FRaycastHit>::value> RaycastVisitor(Start, Dir, HitBuffer, OutputFlags, CommonParams);
 	HitBuffer.IncFlushCount();
 	SpatialAcceleration.Raycast(Start, Dir, DeltaMagnitude, RaycastVisitor);
 	HitBuffer.DecFlushCount();
 }
 
+void FChaosSQAccelerator::Raycast(const FVector& Start, const FVector& Dir, const float DeltaMagnitude, ChaosInterface::FSQHitBuffer<ChaosInterface::FRaycastHit>& HitBuffer, EHitFlags OutputFlags, const ChaosInterface::FSceneQueryCommonParams& CommonParams) const
+{
+	RaycastImp(Start, Dir, DeltaMagnitude, HitBuffer, OutputFlags, CommonParams);
+}
+
+void FChaosSQAccelerator::Raycast(const FVector& Start, const FVector& Dir, const float DeltaMagnitude, ChaosInterface::FSQHitBuffer<ChaosInterface::FPTRaycastHit>& HitBuffer, EHitFlags OutputFlags, const ChaosInterface::FSceneQueryCommonParams& CommonParams) const
+{
+	RaycastImp(Start, Dir, DeltaMagnitude, HitBuffer, OutputFlags, CommonParams);
+}
+
+void FChaosSQAccelerator::Sweep(const Chaos::FImplicitObject& QueryGeom, const FTransform& StartTM, const FVector& Dir, const float DeltaMagnitude, ChaosInterface::FSQHitBuffer<ChaosInterface::FSweepHit>& HitBuffer, EHitFlags OutputFlags, const ChaosInterface::FSceneQueryCommonParams& CommonParams) const
+{
+	return Chaos::Utilities::CastHelper(QueryGeom, StartTM, [this, &Dir, &DeltaMagnitude, &HitBuffer, &OutputFlags, &CommonParams](const auto& Downcast, const FTransform& StartFullTM)
+	{
+		return SweepHelper(Downcast, SpatialAcceleration, StartFullTM, Dir, DeltaMagnitude, HitBuffer, OutputFlags, CommonParams);
+	});
+}
+
+void FChaosSQAccelerator::Sweep(const Chaos::FImplicitObject& QueryGeom, const FTransform& StartTM, const FVector& Dir, const float DeltaMagnitude, ChaosInterface::FSQHitBuffer<ChaosInterface::FPTSweepHit>& HitBuffer, EHitFlags OutputFlags, const ChaosInterface::FSceneQueryCommonParams& CommonParams) const
+{
+	return Chaos::Utilities::CastHelper(QueryGeom, StartTM, [this, &Dir, &DeltaMagnitude, &HitBuffer, &OutputFlags, &CommonParams](const auto& Downcast, const FTransform& StartFullTM)
+	{
+		return SweepHelper(Downcast, SpatialAcceleration, StartFullTM, Dir, DeltaMagnitude, HitBuffer, OutputFlags, CommonParams);
+	});
+}
+
+void FChaosSQAccelerator::Overlap(const Chaos::FImplicitObject& QueryGeom, const FTransform& GeomPose, ChaosInterface::FSQHitBuffer<ChaosInterface::FOverlapHit>& HitBuffer, const ChaosInterface::FSceneQueryCommonParams& CommonParams) const
+{
+	return Chaos::Utilities::CastHelper(QueryGeom, GeomPose, [this, &HitBuffer, &CommonParams](const auto& Downcast, const FTransform& GeomFullPose)
+	{
+		return OverlapHelper(Downcast, SpatialAcceleration, GeomFullPose, HitBuffer, CommonParams);
+	});
+}
+
+void FChaosSQAccelerator::Overlap(const Chaos::FImplicitObject& QueryGeom, const FTransform& GeomPose, ChaosInterface::FSQHitBuffer<ChaosInterface::FPTOverlapHit>& HitBuffer, const ChaosInterface::FSceneQueryCommonParams& CommonParams) const
+{
+	return Chaos::Utilities::CastHelper(QueryGeom, GeomPose, [this, &HitBuffer, &CommonParams](const auto& Downcast, const FTransform& GeomFullPose)
+	{
+		return OverlapHelper(Downcast, SpatialAcceleration, GeomFullPose, HitBuffer, CommonParams);
+	});
+}
+
+PRAGMA_DISABLE_DEPRECATION_WARNINGS
 void FChaosSQAccelerator::Raycast(const FVector& Start, const FVector& Dir, const float DeltaMagnitude, ChaosInterface::FSQHitBuffer<ChaosInterface::FRaycastHit>& HitBuffer, EHitFlags OutputFlags, const ChaosInterface::FQueryFilterData& QueryFilterData, ICollisionQueryFilterCallbackBase& QueryCallback, const ChaosInterface::FQueryDebugParams& DebugParams) const
 {
-	RaycastImp(Start, Dir, DeltaMagnitude, HitBuffer, OutputFlags, QueryFilterData, QueryCallback, DebugParams);
+	const ChaosInterface::FSceneQueryCommonParams CommonData(QueryCallback, QueryFilterData, DebugParams);
+	RaycastImp(Start, Dir, DeltaMagnitude, HitBuffer, OutputFlags, CommonData);
 }
 
 void FChaosSQAccelerator::Raycast(const FVector& Start, const FVector& Dir, const float DeltaMagnitude, ChaosInterface::FSQHitBuffer<ChaosInterface::FPTRaycastHit>& HitBuffer, EHitFlags OutputFlags, const ChaosInterface::FQueryFilterData& QueryFilterData, ICollisionQueryFilterCallbackBase& QueryCallback, const ChaosInterface::FQueryDebugParams& DebugParams) const
 {
-	RaycastImp(Start, Dir, DeltaMagnitude, HitBuffer, OutputFlags, QueryFilterData, QueryCallback, DebugParams);
+	const ChaosInterface::FSceneQueryCommonParams CommonData(QueryCallback, QueryFilterData, DebugParams);
+	RaycastImp(Start, Dir, DeltaMagnitude, HitBuffer, OutputFlags, CommonData);
 }
 
 void FChaosSQAccelerator::Sweep(const Chaos::FImplicitObject& QueryGeom, const FTransform& StartTM, const FVector& Dir, const float DeltaMagnitude, ChaosInterface::FSQHitBuffer<ChaosInterface::FSweepHit>& HitBuffer, EHitFlags OutputFlags, const ChaosInterface::FQueryFilterData& QueryFilterData, ICollisionQueryFilterCallbackBase& QueryCallback, const ChaosInterface::FQueryDebugParams& DebugParams) const
@@ -111,3 +156,4 @@ void FChaosSQAccelerator::Overlap(const Chaos::FImplicitObject& QueryGeom, const
 {
 	return Chaos::Utilities::CastHelper(QueryGeom, GeomPose, [&](const auto& Downcast, const FTransform& GeomFullPose) { return OverlapHelper(Downcast, SpatialAcceleration, GeomFullPose, HitBuffer, QueryFilterData, QueryCallback, DebugParams); });
 }
+PRAGMA_ENABLE_DEPRECATION_WARNINGS
diff --git a/Engine/Source/Runtime/PhysicsCore/Private/SceneQueryCommonParams.cpp b/Engine/Source/Runtime/PhysicsCore/Private/SceneQueryCommonParams.cpp
new file mode 100644
--- /dev/null
+++ b/Engine/Source/Runtime/PhysicsCore/Private/SceneQueryCommonParams.cpp
@@ -0,0 +1,45 @@
+// Copyright Epic Games, Inc. All Rights Reserved.
+
+#include "SceneQueryCommonParams.h"
+
+namespace ChaosInterface
+{
+	FSceneQueryCommonParams::FSceneQueryCommonParams(ICollisionQueryFilterCallback& InQueryCallback)
+		: QueryCallback(&InQueryCallback)
+	{
+	}
+
+	FSceneQueryCommonParams::FSceneQueryCommonParams(ICollisionQueryFilterCallback& InQueryCallback, const FQueryFilterData& InFilterData, FChaosQueryFlags InFlags, const FQueryDebugParams& InDebugParams)
+		: FilterData(InFilterData)
+		, Flags(InFlags)
+		, DebugParams(InDebugParams)
+		, QueryCallback(&InQueryCallback)
+	{
+	}
+
+	FSceneQueryCommonParams::FSceneQueryCommonParams(ICollisionQueryFilterCallback& InQueryCallback, const FQueryFilterData& InFilterData, EQueryFlags InFlags, const FQueryDebugParams& InDebugParams)
+		: FSceneQueryCommonParams(InQueryCallback, InFilterData, U2CQueryFlags(InFlags), InDebugParams)
+	{
+	}
+
+	PRAGMA_DISABLE_DEPRECATION_WARNINGS
+	FSceneQueryCommonParams::FSceneQueryCommonParams(ICollisionQueryFilterCallback& InQueryCallback, const FChaosQueryFilterData& InFilterData, const FQueryDebugParams& InDebugParams)
+		: DebugParams(InDebugParams)
+		, QueryCallback(&InQueryCallback)
+	{
+		FCollisionFilterData CollisionFilterData;
+		CollisionFilterData.Word0 = InFilterData.data.word0;
+		CollisionFilterData.Word1 = InFilterData.data.word1;
+		CollisionFilterData.Word2 = InFilterData.data.word2;
+		CollisionFilterData.Word3 = InFilterData.data.word3;
+
+		FilterData = Chaos::Filter::FQueryFilterBuilder::BuildFromLegacyQueryFilter(CollisionFilterData);
+		Flags = InFilterData.flags;
+	}
+	PRAGMA_ENABLE_DEPRECATION_WARNINGS
+
+	ICollisionQueryFilterCallback& FSceneQueryCommonParams::GetQueryCallback() const
+	{
+		return *QueryCallback;
+	}
+} // namespace ChaosInterface
diff --git a/Engine/Source/Runtime/PhysicsCore/Public/ChaosInterfaceWrapperCore.h b/Engine/Source/Runtime/PhysicsCore/Public/ChaosInterfaceWrapperCore.h
--- a/Engine/Source/Runtime/PhysicsCore/Public/ChaosInterfaceWrapperCore.h
+++ b/Engine/Source/Runtime/PhysicsCore/Public/ChaosInterfaceWrapperCore.h
@@ -31,7 +31,9 @@ struct FDummyPhysActor {};
 template<typename DummyT>
 struct FDummyCallback {};
 
-using FQueryFilterData = FChaosQueryFilterData;
+PRAGMA_DISABLE_DEPRECATION_WARNINGS
+using FQueryFilterData UE_DEPRECATED(5.8, "Use ChaosInterface::FSceneQueryCommonParams instead.") = FChaosQueryFilterData;
+PRAGMA_ENABLE_DEPRECATION_WARNINGS
 
 /** We use this struct so that if no conversion is needed in another API, we can avoid the copy (if we think that's critical) */
 struct FPhysicsRaycastInputAdapater
diff --git a/Engine/Source/Runtime/PhysicsCore/Public/CollisionQueryFilterCallbackCore.h b/Engine/Source/Runtime/PhysicsCore/Public/CollisionQueryFilterCallbackCore.h
--- a/Engine/Source/Runtime/PhysicsCore/Public/CollisionQueryFilterCallbackCore.h
+++ b/Engine/Source/Runtime/PhysicsCore/Public/CollisionQueryFilterCallbackCore.h
@@ -21,7 +21,46 @@ namespace Chaos
 	class FImplicitObject;
 }
 
-class ICollisionQueryFilterCallbackBase
+class ICollisionQueryFilterCallback
+{
+public:
+	using FQueryFilterData = Chaos::Filter::FQueryFilterData;
+
+	virtual ~ICollisionQueryFilterCallback() = default;
+	
+	virtual ECollisionQueryHitType PreFilter(const FQueryFilterData& FilterData, const Chaos::FPerShapeData& Shape, const Chaos::FGeometryParticle& Actor) = 0;
+	virtual ECollisionQueryHitType PreFilter(const FQueryFilterData& FilterData, const Chaos::FPerShapeData& Shape, const Chaos::FGeometryParticleHandle& Actor) = 0;
+
+	virtual ECollisionQueryHitType PostFilter(const FQueryFilterData& FilterData, const ChaosInterface::FQueryHit& Hit) = 0;
+	virtual ECollisionQueryHitType PostFilter(const FQueryFilterData& FilterData, const ChaosInterface::FPTQueryHit& Hit) = 0;
+};
+
+class FBlockAllCollisionQueryFilterCallback : public ICollisionQueryFilterCallback
+{
+public:
+	virtual ~FBlockAllCollisionQueryFilterCallback() = default;
+
+	virtual ECollisionQueryHitType PreFilter(const FQueryFilterData& FilterData, const Chaos::FPerShapeData& Shape, const Chaos::FGeometryParticle& Actor) { return ECollisionQueryHitType::Block; }
+	virtual ECollisionQueryHitType PreFilter(const FQueryFilterData& FilterData, const Chaos::FPerShapeData& Shape, const Chaos::FGeometryParticleHandle& Actor) { return ECollisionQueryHitType::Block; }
+
+	virtual ECollisionQueryHitType PostFilter(const FQueryFilterData& FilterData, const ChaosInterface::FQueryHit& Hit) { return ECollisionQueryHitType::Block; }
+	virtual ECollisionQueryHitType PostFilter(const FQueryFilterData& FilterData, const ChaosInterface::FPTQueryHit& Hit) { return ECollisionQueryHitType::Block; }
+};
+
+class FOverlapAllCollisionQueryFilterCallback : public ICollisionQueryFilterCallback
+{
+public:
+	virtual ~FOverlapAllCollisionQueryFilterCallback() = default;
+
+	virtual ECollisionQueryHitType PreFilter(const FQueryFilterData& FilterData, const Chaos::FPerShapeData& Shape, const Chaos::FGeometryParticle& Actor) { return ECollisionQueryHitType::Touch; }
+	virtual ECollisionQueryHitType PreFilter(const FQueryFilterData& FilterData, const Chaos::FPerShapeData& Shape, const Chaos::FGeometryParticleHandle& Actor) { return ECollisionQueryHitType::Touch; }
+
+	virtual ECollisionQueryHitType PostFilter(const FQueryFilterData& FilterData, const ChaosInterface::FQueryHit& Hit) { return ECollisionQueryHitType::Touch; }
+	virtual ECollisionQueryHitType PostFilter(const FQueryFilterData& FilterData, const ChaosInterface::FPTQueryHit& Hit) { return ECollisionQueryHitType::Touch; }
+};
+
+PRAGMA_DISABLE_DEPRECATION_WARNINGS
+class UE_DEPRECATED(5.8, "Use ICollisionQueryFilterCallback instead") ICollisionQueryFilterCallbackBase : public ICollisionQueryFilterCallback
 {
 public:
 	
@@ -31,9 +70,35 @@ public:
 
 	virtual ECollisionQueryHitType PostFilter(const FCollisionFilterData& FilterData, const ChaosInterface::FPTQueryHit& Hit) = 0;
 	virtual ECollisionQueryHitType PreFilter(const FCollisionFilterData& FilterData, const Chaos::FPerShapeData& Shape, const Chaos::FGeometryParticleHandle& Actor) = 0;
+
+private:
+	static FCollisionFilterData GetLegacyFilterData(const FQueryFilterData& FilterData)
+	{
+		return Chaos::Filter::FQueryFilterBuilder::GetLegacyQueryFilter(FilterData);
+	}
+
+	virtual ECollisionQueryHitType PostFilter(const FQueryFilterData& FilterData, const ChaosInterface::FQueryHit& Hit) override
+	{
+		return PostFilter(GetLegacyFilterData(FilterData), Hit);
+	}
+
+	virtual ECollisionQueryHitType PreFilter(const FQueryFilterData& FilterData, const Chaos::FPerShapeData& Shape, const Chaos::FGeometryParticle& Actor) override
+	{
+		return PreFilter(GetLegacyFilterData(FilterData), Shape, Actor);
+	}
+
+	virtual ECollisionQueryHitType PostFilter(const FQueryFilterData& FilterData, const ChaosInterface::FPTQueryHit& Hit) override
+	{
+		return PostFilter(GetLegacyFilterData(FilterData), Hit);
+	}
+
+	virtual ECollisionQueryHitType PreFilter(const FQueryFilterData& FilterData, const Chaos::FPerShapeData& Shape, const Chaos::FGeometryParticleHandle& Actor) override
+	{
+		return PreFilter(GetLegacyFilterData(FilterData), Shape, Actor);
+	}
 };
 
-class FBlockAllQueryCallback : public ICollisionQueryFilterCallbackBase
+class UE_DEPRECATED(5.8, "Use FBlockAllCollisionQueryFilterCallback instead") FBlockAllQueryCallback : public ICollisionQueryFilterCallbackBase
 {
 public:
 	virtual ~FBlockAllQueryCallback() {}
@@ -44,7 +109,7 @@ public:
 	virtual ECollisionQueryHitType PreFilter(const FCollisionFilterData& FilterData, const Chaos::FPerShapeData& Shape, const Chaos::FGeometryParticleHandle& Actor) override { return ECollisionQueryHitType::Block; }
 };
 
-class FOverlapAllQueryCallback : public ICollisionQueryFilterCallbackBase
+class UE_DEPRECATED(5.8, "Use FOverlapAllCollisionQueryFilterCallback instead") FOverlapAllQueryCallback : public ICollisionQueryFilterCallbackBase
 {
 public:
 	virtual ~FOverlapAllQueryCallback() {}
@@ -54,3 +119,4 @@ public:
 	virtual ECollisionQueryHitType PostFilter(const FCollisionFilterData& FilterData, const ChaosInterface::FPTQueryHit& Hit) override { return ECollisionQueryHitType::Touch; }
 	virtual ECollisionQueryHitType PreFilter(const FCollisionFilterData& FilterData, const Chaos::FPerShapeData& Shape, const Chaos::FGeometryParticleHandle& Actor) override { return ECollisionQueryHitType::Touch; }
 };
+PRAGMA_ENABLE_DEPRECATION_WARNINGS
diff --git a/Engine/Source/Runtime/PhysicsCore/Public/PhysicsInterfaceTypesCore.h b/Engine/Source/Runtime/PhysicsCore/Public/PhysicsInterfaceTypesCore.h
--- a/Engine/Source/Runtime/PhysicsCore/Public/PhysicsInterfaceTypesCore.h
+++ b/Engine/Source/Runtime/PhysicsCore/Public/PhysicsInterfaceTypesCore.h
@@ -263,7 +263,7 @@ template <> FORCEINLINE void FChaosUserData::Set(void* UserData, FPhysicsObject*
 template <> FORCEINLINE void FChaosUserData::Set(void* UserData, FChaosUserEntityAppend* Payload) { check(UserData); ((FChaosUserData*)UserData)->Type = EChaosUserDataType::ChaosUserEntity; ((FChaosUserData*)UserData)->Payload = Payload; }
 template <> FORCEINLINE void FChaosUserData::Set(void* UserData, FCustomChaosPayload* Payload) { check(UserData); ((FChaosUserData*)UserData)->Type = EChaosUserDataType::CustomPayload; ((FChaosUserData*)UserData)->Payload = Payload; }
 
-struct FChaosFilterData
+struct UE_DEPRECATED(5.8, "Use Chaos::Filter::FQueryFilterData instead") FChaosFilterData
 {
 	FORCEINLINE FChaosFilterData()
 	{
@@ -403,7 +403,8 @@ inline FChaosQueryFlags U2CQueryFlags(FQueryFlags Flags)
 	return (FChaosQueryFlags)Result;
 }
 
-struct FChaosQueryFilterData
+PRAGMA_DISABLE_DEPRECATION_WARNINGS
+struct UE_DEPRECATED(5.8, "Use ChaosInterface::FSceneQueryCommonParams instead.") FChaosQueryFilterData
 {
 	explicit FORCEINLINE FChaosQueryFilterData() : flags(FChaosQueryFlag::eDYNAMIC | FChaosQueryFlag::eSTATIC), clientId(0) {}
 
@@ -415,6 +416,7 @@ struct FChaosQueryFilterData
 	FChaosQueryFlags	flags;
 	uint8 clientId;
 };
+PRAGMA_ENABLE_DEPRECATION_WARNINGS
 
 template <typename enumtype, typename storagetype>
 FORCEINLINE ChaosFlags<enumtype, storagetype>::ChaosFlags(void)
diff --git a/Engine/Source/Runtime/PhysicsCore/Public/PhysicsInterfaceUtilsCore.h b/Engine/Source/Runtime/PhysicsCore/Public/PhysicsInterfaceUtilsCore.h
--- a/Engine/Source/Runtime/PhysicsCore/Public/PhysicsInterfaceUtilsCore.h
+++ b/Engine/Source/Runtime/PhysicsCore/Public/PhysicsInterfaceUtilsCore.h
@@ -9,10 +9,12 @@
 
 namespace ChaosInterface { struct FPTLocationHit; }
 
+PRAGMA_DISABLE_DEPRECATION_WARNINGS
 UE_INTERNAL PHYSICSCORE_API FCollisionFilterData C2UFilterData(const FChaosFilterData& FilterData);
 UE_INTERNAL PHYSICSCORE_API FChaosFilterData U2CFilterData(const FCollisionFilterData& FilterData);
 
 UE_INTERNAL PHYSICSCORE_API FCollisionFilterData ToUnrealFilterData(const FChaosFilterData& FilterData);
+PRAGMA_ENABLE_DEPRECATION_WARNINGS
 
 //Find the face index for a given hit. This gives us a chance to modify face index based on things like most opposing normal
 
diff --git a/Engine/Source/Runtime/PhysicsCore/Public/SQAccelerator.h b/Engine/Source/Runtime/PhysicsCore/Public/SQAccelerator.h
--- a/Engine/Source/Runtime/PhysicsCore/Public/SQAccelerator.h
+++ b/Engine/Source/Runtime/PhysicsCore/Public/SQAccelerator.h
@@ -11,6 +11,7 @@
 #include "Math/Transform.h"
 #include "PhysicsInterfaceDeclaresCore.h"
 #include "PhysicsInterfaceWrapperShared.h"
+#include "SceneQueryCommonParams.h"
 
 namespace ChaosInterface { struct FOverlapHit; }
 namespace ChaosInterface { struct FPTOverlapHit; }
@@ -41,21 +42,38 @@ public:
 	PHYSICSCORE_API FChaosSQAccelerator(const Chaos::ISpatialAcceleration<Chaos::FAccelerationStructureHandle,Chaos::FReal, 3>& InSpatialAcceleration);
 	virtual ~FChaosSQAccelerator() {};
 
+	PHYSICSCORE_API void Raycast(const FVector& Start, const FVector& Dir, const float DeltaMagnitude, ChaosInterface::FSQHitBuffer<ChaosInterface::FRaycastHit>& HitBuffer, EHitFlags OutputFlags, const ChaosInterface::FSceneQueryCommonParams& CommonParams) const;
+	PHYSICSCORE_API void Raycast(const FVector& Start, const FVector& Dir, const float DeltaMagnitude, ChaosInterface::FSQHitBuffer<ChaosInterface::FPTRaycastHit>& HitBuffer, EHitFlags OutputFlags, const ChaosInterface::FSceneQueryCommonParams& CommonParams) const;
+
+	PHYSICSCORE_API void Sweep(const Chaos::FImplicitObject& QueryGeom, const FTransform& StartTM, const FVector& Dir, const float DeltaMagnitude, ChaosInterface::FSQHitBuffer<ChaosInterface::FSweepHit>& HitBuffer, EHitFlags OutputFlags, const ChaosInterface::FSceneQueryCommonParams& CommonParams) const;
+	PHYSICSCORE_API void Sweep(const Chaos::FImplicitObject& QueryGeom, const FTransform& StartTM, const FVector& Dir, const float DeltaMagnitude, ChaosInterface::FSQHitBuffer<ChaosInterface::FPTSweepHit>& HitBuffer, EHitFlags OutputFlags, const ChaosInterface::FSceneQueryCommonParams& CommonParams) const;
+
+	PHYSICSCORE_API void Overlap(const Chaos::FImplicitObject& QueryGeom, const FTransform& GeomPose, ChaosInterface::FSQHitBuffer<ChaosInterface::FOverlapHit>& HitBuffer, const ChaosInterface::FSceneQueryCommonParams& CommonParams) const;
+	PHYSICSCORE_API void Overlap(const Chaos::FImplicitObject& QueryGeom, const FTransform& GeomPose, ChaosInterface::FSQHitBuffer<ChaosInterface::FPTOverlapHit>& HitBuffer, const ChaosInterface::FSceneQueryCommonParams& CommonParams) const;
+
+
+	PRAGMA_DISABLE_DEPRECATION_WARNINGS
+	UE_DEPRECATED(5.8, "Use the overload that takes ChaosInterface::FSceneQueryCommonParams instead")
 	PHYSICSCORE_API void Raycast(const FVector& Start, const FVector& Dir, const float DeltaMagnitude, ChaosInterface::FSQHitBuffer<ChaosInterface::FRaycastHit>& HitBuffer, EHitFlags OutputFlags, const ChaosInterface::FQueryFilterData& QueryFilterData, ICollisionQueryFilterCallbackBase& QueryCallback, const ChaosInterface::FQueryDebugParams& DebugParams = {}) const;
+	UE_DEPRECATED(5.8, "Use the overload that takes ChaosInterface::FSceneQueryCommonParams instead")
 	PHYSICSCORE_API void Raycast(const FVector& Start, const FVector& Dir, const float DeltaMagnitude, ChaosInterface::FSQHitBuffer<ChaosInterface::FPTRaycastHit>& HitBuffer, EHitFlags OutputFlags, const ChaosInterface::FQueryFilterData& QueryFilterData, ICollisionQueryFilterCallbackBase& QueryCallback, const ChaosInterface::FQueryDebugParams& DebugParams = {}) const;
 
+	UE_DEPRECATED(5.8, "Use the overload that takes ChaosInterface::FSceneQueryCommonParams instead")
 	PHYSICSCORE_API void Sweep(const Chaos::FImplicitObject& QueryGeom, const FTransform& StartTM, const FVector& Dir, const float DeltaMagnitude, ChaosInterface::FSQHitBuffer<ChaosInterface::FSweepHit>& HitBuffer, EHitFlags OutputFlags, const ChaosInterface::FQueryFilterData& QueryFilterData, ICollisionQueryFilterCallbackBase& QueryCallback, const ChaosInterface::FQueryDebugParams& DebugParams = {}) const;
+	UE_DEPRECATED(5.8, "Use the overload that takes ChaosInterface::FSceneQueryCommonParams instead")
 	PHYSICSCORE_API void Sweep(const Chaos::FImplicitObject& QueryGeom, const FTransform& StartTM, const FVector& Dir, const float DeltaMagnitude, ChaosInterface::FSQHitBuffer<ChaosInterface::FPTSweepHit>& HitBuffer, EHitFlags OutputFlags, const ChaosInterface::FQueryFilterData& QueryFilterData, ICollisionQueryFilterCallbackBase& QueryCallback, const ChaosInterface::FQueryDebugParams& DebugParams = {}) const;
 
+	UE_DEPRECATED(5.8, "Use the overload that takes ChaosInterface::FSceneQueryCommonParams instead")
 	PHYSICSCORE_API void Overlap(const Chaos::FImplicitObject& QueryGeom, const FTransform& GeomPose, ChaosInterface::FSQHitBuffer<ChaosInterface::FOverlapHit>& HitBuffer, const ChaosInterface::FQueryFilterData& QueryFilterData, ICollisionQueryFilterCallbackBase& QueryCallback, const ChaosInterface::FQueryDebugParams& DebugParams = {}) const;
+	UE_DEPRECATED(5.8, "Use the overload that takes ChaosInterface::FSceneQueryCommonParams instead")
 	PHYSICSCORE_API void Overlap(const Chaos::FImplicitObject& QueryGeom, const FTransform& GeomPose, ChaosInterface::FSQHitBuffer<ChaosInterface::FPTOverlapHit>& HitBuffer, const ChaosInterface::FQueryFilterData& QueryFilterData, ICollisionQueryFilterCallbackBase& QueryCallback, const ChaosInterface::FQueryDebugParams& DebugParams = {}) const;
+	PRAGMA_ENABLE_DEPRECATION_WARNINGS
 
 private:
 	const Chaos::ISpatialAcceleration<Chaos::FAccelerationStructureHandle, Chaos::FReal, 3>& SpatialAcceleration;
 
 	template <typename TRaycastHit>
-	void RaycastImp(const FVector& Start, const FVector& Dir, const float DeltaMagnitude, ChaosInterface::FSQHitBuffer<TRaycastHit>& HitBuffer, EHitFlags OutputFlags, const ChaosInterface::FQueryFilterData& QueryFilterData, ICollisionQueryFilterCallbackBase& QueryCallback, const ChaosInterface::FQueryDebugParams& DebugParams) const;
-
+	void RaycastImp(const FVector& Start, const FVector& Dir, const float DeltaMagnitude, ChaosInterface::FSQHitBuffer<TRaycastHit>& HitBuffer, EHitFlags OutputFlags, const ChaosInterface::FSceneQueryCommonParams& CommonParams) const;
 };
 
 PRAGMA_DISABLE_DEPRECATION_WARNINGS
diff --git a/Engine/Source/Runtime/PhysicsCore/Public/SQCapture.h b/Engine/Source/Runtime/PhysicsCore/Public/SQCapture.h
--- a/Engine/Source/Runtime/PhysicsCore/Public/SQCapture.h
+++ b/Engine/Source/Runtime/PhysicsCore/Public/SQCapture.h
@@ -37,6 +37,7 @@ namespace Chaos
 }
 
 //Allows us to capture a scene query with either physx or chaos and then load it into either format for testing purposes
+PRAGMA_DISABLE_DEPRECATION_WARNINGS
 struct UE_INTERNAL FSQCapture
 {
 	PHYSICSCORE_API ~FSQCapture();
@@ -114,3 +115,4 @@ private:
 	bool bChaosDataReady;
 	bool bPhysXDataReady;
 };
+PRAGMA_ENABLE_DEPRECATION_WARNINGS
diff --git a/Engine/Source/Runtime/PhysicsCore/Public/SQVisitor.h b/Engine/Source/Runtime/PhysicsCore/Public/SQVisitor.h
--- a/Engine/Source/Runtime/PhysicsCore/Public/SQVisitor.h
+++ b/Engine/Source/Runtime/PhysicsCore/Public/SQVisitor.h
@@ -11,6 +11,7 @@
 #include "CollisionQueryFilterCallbackCore.h"
 #include "PhysicsInterfaceTypesCore.h"
 #include "PhysicsInterfaceUtilsCore.h"
+#include "SceneQueryCommonParams.h"
 
 #include <type_traits>
 
@@ -45,74 +46,104 @@ struct TSQVisitor : public Chaos::ISpatialVisitor<TPayload, Chaos::FReal>
 	using FAABB3 = Chaos::FAABB3;
 	using FVec3 = Chaos::FVec3;
 	TSQVisitor(const FVector& InStartPoint, const FVector& InDir, ChaosInterface::FSQHitBuffer<THitType>& InHitBuffer, EHitFlags InOutputFlags,
-		const ChaosInterface::FQueryFilterData& InQueryFilterData, ICollisionQueryFilterCallbackBase& InQueryCallback, const ChaosInterface::FQueryDebugParams& InDebugParams)
+		const ChaosInterface::FSceneQueryCommonParams& InCommonParams)
 		: StartPoint(InStartPoint)
 		, Dir(InDir)
 		, OutputFlags(InOutputFlags)
 		, bAnyHit(false)
-		, DebugParams(InDebugParams)
+		, DebugParams(InCommonParams.DebugParams)
 		, HitFaceNormal(Chaos::FVec3::ZeroVector)
 		, HitBuffer(InHitBuffer)
-		, QueryFilterData(InQueryFilterData)
-		, QueryFilterDataConcrete(C2UFilterData(QueryFilterData.data))
-		, QueryCallback(InQueryCallback)
+		, QueryFlags(InCommonParams.Flags)
+		, QueryFilterData(InCommonParams.FilterData)
+		, LegacyQueryFilterDataConcrete(Chaos::Filter::FQueryFilterBuilder::GetLegacyQueryFilter(InCommonParams.FilterData))
+		, QueryCallback(InCommonParams.GetQueryCallback())
 		, QueryCenter(StartTM.GetLocation())
 	{
 		QueryGeomWorldBounds = FAABB3(QueryCenter, QueryCenter);
-		bAnyHit = QueryFilterData.flags & FChaosQueryFlag::eANY_HIT;
+		bAnyHit = QueryFlags & FChaosQueryFlag::eANY_HIT;
 	}
 
+	PRAGMA_DISABLE_DEPRECATION_WARNINGS
+	UE_DEPRECATED(5.8, "Use the overload that takes a ChaosInterface::FSceneQueryCommonParams instead")
+	TSQVisitor(const FVector& InStartPoint, const FVector& InDir, ChaosInterface::FSQHitBuffer<THitType>& InHitBuffer, EHitFlags InOutputFlags,
+		const ChaosInterface::FQueryFilterData& InQueryFilterData, ICollisionQueryFilterCallbackBase& InQueryCallback, const ChaosInterface::FQueryDebugParams& InDebugParams)
+		: TSQVisitor(InStartPoint, InDir, InHitBuffer, InOutputFlags, ChaosInterface::FSceneQueryCommonParams(InQueryCallback, InQueryFilterData, InDebugParams))
+	{
+	}
+	PRAGMA_ENABLE_DEPRECATION_WARNINGS
+
 	TSQVisitor(const FTransform& InStartTM, const FVector& InDir, ChaosInterface::FSQHitBuffer<THitType>& InHitBuffer, EHitFlags InOutputFlags,
-		const ChaosInterface::FQueryFilterData& InQueryFilterData, ICollisionQueryFilterCallbackBase& InQueryCallback, const QueryGeometryType& InQueryGeom, const FAABB3& InQueryGeomWorldBounds, const ChaosInterface::FQueryDebugParams& InDebugParams)
+		const QueryGeometryType& InQueryGeom, const FAABB3& InQueryGeomWorldBounds, const ChaosInterface::FSceneQueryCommonParams& InCommonParams)
 		: Dir(InDir)
 		, OutputFlags(InOutputFlags)
 		, bAnyHit(false)
-		, DebugParams(InDebugParams)
+		, DebugParams(InCommonParams.DebugParams)
 		, HitFaceNormal(Chaos::FVec3::ZeroVector)
 		, HitBuffer(InHitBuffer)
-		, QueryFilterData(InQueryFilterData)
-		, QueryFilterDataConcrete(C2UFilterData(QueryFilterData.data))
+		, QueryFlags(InCommonParams.Flags)
+		, QueryFilterData(InCommonParams.FilterData)
+		, LegacyQueryFilterDataConcrete(Chaos::Filter::FQueryFilterBuilder::GetLegacyQueryFilter(InCommonParams.FilterData))
 		, QueryGeom(&InQueryGeom)
-		, QueryCallback(InQueryCallback)
+		, QueryCallback(InCommonParams.GetQueryCallback())
 		, StartTM(InStartTM)
 		, QueryGeomWorldBounds(InQueryGeomWorldBounds)
 		, QueryCenter(InQueryGeomWorldBounds.Center())
 	{
-		bAnyHit = QueryFilterData.flags & FChaosQueryFlag::eANY_HIT;
+		bAnyHit = QueryFlags & FChaosQueryFlag::eANY_HIT;
 		//todo: check THitType is sweep
 	}
 
+	PRAGMA_DISABLE_DEPRECATION_WARNINGS
+	UE_DEPRECATED(5.8, "Use the overload that takes a ChaosInterface::FSceneQueryCommonParams instead")
+	TSQVisitor(const FTransform& InStartTM, const FVector& InDir, ChaosInterface::FSQHitBuffer<THitType>& InHitBuffer, EHitFlags InOutputFlags,
+		const ChaosInterface::FQueryFilterData& InQueryFilterData, ICollisionQueryFilterCallbackBase& InQueryCallback, const QueryGeometryType& InQueryGeom, const FAABB3& InQueryGeomWorldBounds, const ChaosInterface::FQueryDebugParams& InDebugParams)
+		: TSQVisitor(InStartTM, InDir, InHitBuffer, InOutputFlags, InQueryGeom, InQueryGeomWorldBounds, ChaosInterface::FSceneQueryCommonParams(InQueryCallback, InQueryFilterData, InDebugParams))
+	{
+	}
+
 	UE_DEPRECATED(5.8, "Use the overload that takes the QueryGeomWorldBounds")
 	TSQVisitor(const FTransform& InStartTM, const FVector& InDir, ChaosInterface::FSQHitBuffer<THitType>& InHitBuffer, EHitFlags InOutputFlags,
 		const ChaosInterface::FQueryFilterData& InQueryFilterData, ICollisionQueryFilterCallbackBase& InQueryCallback, const QueryGeometryType& InQueryGeom, const ChaosInterface::FQueryDebugParams& InDebugParams)
 		: TSQVisitor(InStartTM, InDir, InHitBuffer, InOutputFlags, InQueryFilterData, InQueryCallback, InQueryGeom, InQueryGeom.CalculateTransformedBounds(InStartTM), InDebugParams)
 	{
 	}
+	PRAGMA_ENABLE_DEPRECATION_WARNINGS
 
 	TSQVisitor(const FTransform& InWorldTM, ChaosInterface::FSQHitBuffer<THitType>& InHitBuffer,
-		const ChaosInterface::FQueryFilterData& InQueryFilterData, ICollisionQueryFilterCallbackBase& InQueryCallback, const QueryGeometryType& InQueryGeom, const FAABB3& InQueryGeomWorldBounds, const ChaosInterface::FQueryDebugParams& InDebugParams)
+		const QueryGeometryType& InQueryGeom, const FAABB3& InQueryGeomWorldBounds, const ChaosInterface::FSceneQueryCommonParams& InCommonParams)
 		: bAnyHit(false)
-		, DebugParams(InDebugParams)
+		, DebugParams(InCommonParams.DebugParams)
 		, HitFaceNormal(Chaos::FVec3::ZeroVector)
 		, HitBuffer(InHitBuffer)
-		, QueryFilterData(InQueryFilterData)
-		, QueryFilterDataConcrete(C2UFilterData(QueryFilterData.data))
+		, QueryFlags(InCommonParams.Flags)
+		, QueryFilterData(InCommonParams.FilterData)
+		, LegacyQueryFilterDataConcrete(Chaos::Filter::FQueryFilterBuilder::GetLegacyQueryFilter(InCommonParams.FilterData))
 		, QueryGeom(&InQueryGeom)
-		, QueryCallback(InQueryCallback)
+		, QueryCallback(InCommonParams.GetQueryCallback())
 		, StartTM(InWorldTM)
 		, QueryGeomWorldBounds(InQueryGeomWorldBounds)
 		, QueryCenter(InQueryGeomWorldBounds.Center())
 	{
-		bAnyHit = QueryFilterData.flags & FChaosQueryFlag::eANY_HIT;
+		bAnyHit = QueryFlags & FChaosQueryFlag::eANY_HIT;
 		//todo: check THitType is overlap
 	}
 
+	PRAGMA_DISABLE_DEPRECATION_WARNINGS
+	UE_DEPRECATED(5.8, "Use the overload that takes a ChaosInterface::FSceneQueryCommonParams instead")
+	TSQVisitor(const FTransform& InWorldTM, ChaosInterface::FSQHitBuffer<THitType>& InHitBuffer,
+		const ChaosInterface::FQueryFilterData& InQueryFilterData, ICollisionQueryFilterCallbackBase& InQueryCallback, const QueryGeometryType& InQueryGeom, const FAABB3& InQueryGeomWorldBounds, const ChaosInterface::FQueryDebugParams& InDebugParams)
+		: TSQVisitor(InWorldTM, InHitBuffer, InQueryGeom, InQueryGeomWorldBounds, ChaosInterface::FSceneQueryCommonParams(InQueryCallback, InQueryFilterData, InDebugParams))
+	{
+	}
+
 	UE_DEPRECATED(5.8, "Use the overload that takes the QueryGeomWorldBounds")
 	TSQVisitor(const FTransform& InWorldTM, ChaosInterface::FSQHitBuffer<THitType>& InHitBuffer,
 		const ChaosInterface::FQueryFilterData& InQueryFilterData, ICollisionQueryFilterCallbackBase& InQueryCallback, const QueryGeometryType& InQueryGeom, const ChaosInterface::FQueryDebugParams& InDebugParams)
 		: TSQVisitor(InWorldTM, InHitBuffer, InQueryFilterData, InQueryCallback, InQueryGeom, InQueryGeom.CalculateTransformedBounds(InWorldTM), InDebugParams)
 	{
 	}
+	PRAGMA_ENABLE_DEPRECATION_WARNINGS
 
 	virtual ~TSQVisitor() {}
 
@@ -133,7 +164,7 @@ struct TSQVisitor : public Chaos::ISpatialVisitor<TPayload, Chaos::FReal>
 
 	virtual const void* GetQueryData() const override
 	{
-		return &QueryFilterDataConcrete;
+		return &LegacyQueryFilterDataConcrete;
 	}
 
 	virtual const void* GetSimData() const override
@@ -273,7 +304,7 @@ private:
 			}
 
 			//TODO: use gt particles directly
-			ECollisionQueryHitType HitType = QueryFilterData.flags & FChaosQueryFlag::ePREFILTER ? QueryCallback.PreFilter(QueryFilterDataConcrete, *Shape, *GeometryParticle) : ECollisionQueryHitType::Block;
+			ECollisionQueryHitType HitType = QueryFlags & FChaosQueryFlag::ePREFILTER ? QueryCallback.PreFilter(QueryFilterData, *Shape, *GeometryParticle) : ECollisionQueryHitType::Block;
 
 			if (HitType != ECollisionQueryHitType::None)
 			{
@@ -392,7 +423,7 @@ private:
 							Hit.FaceNormal = HitFaceNormal;
 						}
 
-						HitType = QueryFilterData.flags & FChaosQueryFlag::ePOSTFILTER ? QueryCallback.PostFilter(QueryFilterDataConcrete, Hit) : HitType;
+						HitType = QueryFlags & FChaosQueryFlag::ePOSTFILTER ? QueryCallback.PostFilter(QueryFilterData, Hit) : HitType;
 
 						if (HitType != ECollisionQueryHitType::None)
 						{
@@ -497,10 +528,11 @@ private:
 	const ChaosInterface::FQueryDebugParams DebugParams;
 	Chaos::FVec3 HitFaceNormal;
 	ChaosInterface::FSQHitBuffer<THitType>& HitBuffer;
-	const ChaosInterface::FQueryFilterData& QueryFilterData;
-	const FCollisionFilterData QueryFilterDataConcrete;
+	const FChaosQueryFlags QueryFlags;
+	const Chaos::Filter::FQueryFilterData QueryFilterData;
+	const FCollisionFilterData LegacyQueryFilterDataConcrete;
 	const QueryGeometryType* QueryGeom = nullptr;
-	ICollisionQueryFilterCallbackBase& QueryCallback;
+	ICollisionQueryFilterCallback& QueryCallback;
 	const FTransform StartTM;
 	FAABB3 QueryGeomWorldBounds = FAABB3::ZeroAABB();
 	FVec3 QueryCenter = FVec3::ZeroVector;
@@ -513,20 +545,30 @@ struct TBPVisitor : public Chaos::ISpatialVisitor<TPayload, Chaos::FReal>
 	using TGeometryType = std::conditional_t<bGTData, Chaos::FGeometryParticle, Chaos::FGeometryParticleHandle>;
 
 	TBPVisitor(const FTransform& InWorldTM, ChaosInterface::FSQHitBuffer<THitType>& InHitBuffer,
-		const ChaosInterface::FQueryFilterData& InQueryFilterData, ICollisionQueryFilterCallbackBase& InQueryCallback, const QueryGeometryType& InQueryGeom, const ChaosInterface::FQueryDebugParams& InDebugParams)
+		const QueryGeometryType& InQueryGeom, const ChaosInterface::FSceneQueryCommonParams& InCommonParams)
 		: HalfExtents(InQueryGeom.BoundingBox().Extents() * 0.5)
 		, bAnyHit(false)
-		, DebugParams(InDebugParams)
+		, DebugParams(InCommonParams.DebugParams)
 		, HitBuffer(InHitBuffer)
-		, QueryFilterData(InQueryFilterData)
-		, QueryFilterDataConcrete(ToUnrealFilterData(InQueryFilterData.data))
+		, QueryFlags(InCommonParams.Flags)
+		, QueryFilterData(InCommonParams.FilterData)
+		, LegacyQueryFilterDataConcrete(Chaos::Filter::FQueryFilterBuilder::GetLegacyQueryFilter(InCommonParams.FilterData))
 		, QueryGeom(&InQueryGeom)
-		, QueryCallback(InQueryCallback)
+		, QueryCallback(InCommonParams.GetQueryCallback())
 		, StartTM(InWorldTM)
 	{
-		bAnyHit = QueryFilterData.flags & FPhysicsQueryFlag::eANY_HIT;
+		bAnyHit = QueryFlags & FPhysicsQueryFlag::eANY_HIT;
 	}
 
+	PRAGMA_DISABLE_DEPRECATION_WARNINGS
+	UE_DEPRECATED(5.8, "Use the overload that takes a ChaosInterface::FSceneQueryCommonParams instead")
+	TBPVisitor(const FTransform& InWorldTM, ChaosInterface::FSQHitBuffer<THitType>& InHitBuffer,
+		const ChaosInterface::FQueryFilterData& InQueryFilterData, ICollisionQueryFilterCallbackBase& InQueryCallback, const QueryGeometryType& InQueryGeom, const ChaosInterface::FQueryDebugParams& InDebugParams)
+		: TBPVisitor(InWorldTM, InHitBuffer, InQueryGeom, ChaosInterface::FSceneQueryCommonParams(InQueryCallback, InQueryFilterData, InDebugParams))
+	{
+	}
+	PRAGMA_ENABLE_DEPRECATION_WARNINGS
+
 	virtual ~TBPVisitor() {}
 
 	virtual bool Overlap(const Chaos::TSpatialVisitorData<TPayload>& Instance) override
@@ -548,7 +590,7 @@ struct TBPVisitor : public Chaos::ISpatialVisitor<TPayload, Chaos::FReal>
 
 	virtual const void* GetQueryData() const override
 	{
-		return &QueryFilterDataConcrete;
+		return &LegacyQueryFilterDataConcrete;
 	}
 
 	virtual const void* GetSimData() const override
@@ -612,7 +654,7 @@ private:
 
 			CVD_TRACE_SCOPED_SCENE_QUERY_VISIT_HELPER(EChaosVDSceneQueryVisitorType::BroadPhase, FTransform(GeometryParticle->GetR(), GeometryParticle->GetX()), Payload.UniqueIdx().Idx, ShapeIndex, CurData);
 
-			ECollisionQueryHitType HitType = QueryFilterData.flags & FPhysicsQueryFlag::ePREFILTER ? QueryCallback.PreFilter(QueryFilterDataConcrete, *Shape, *GeometryParticle) : ECollisionQueryHitType::Block;
+			ECollisionQueryHitType HitType = QueryFlags & FPhysicsQueryFlag::ePREFILTER ? QueryCallback.PreFilter(QueryFilterData, *Shape, *GeometryParticle) : ECollisionQueryHitType::Block;
 			if (HitType != ECollisionQueryHitType::None)
 			{
 				const bool bBlocker = (HitType == ECollisionQueryHitType::Block || bAnyHit || HitBuffer.WantsSingleResult());
@@ -642,16 +684,17 @@ private:
 	bool bAnyHit;
 	const ChaosInterface::FQueryDebugParams DebugParams;
 	ChaosInterface::FSQHitBuffer<THitType>& HitBuffer;
-	const ChaosInterface::FQueryFilterData& QueryFilterData;
-	const FCollisionFilterData QueryFilterDataConcrete;
+	const FChaosQueryFlags QueryFlags;
+	const Chaos::Filter::FQueryFilterData QueryFilterData;
+	const FCollisionFilterData LegacyQueryFilterDataConcrete;
 	const QueryGeometryType* QueryGeom;
-	ICollisionQueryFilterCallbackBase& QueryCallback;
+	ICollisionQueryFilterCallback& QueryCallback;
 	const FTransform StartTM;
 };
 
 
 template <typename QueryGeomType, typename TSweepHit, typename TPayload>
-void SweepHelper(const QueryGeomType& QueryGeom, const Chaos::ISpatialAcceleration<TPayload, Chaos::FReal, 3>& SpatialAcceleration, const FTransform& StartTM, const FVector& Dir, const float DeltaMagnitude, ChaosInterface::FSQHitBuffer<TSweepHit>& HitBuffer, EHitFlags OutputFlags, const ChaosInterface::FQueryFilterData& QueryFilterData, ICollisionQueryFilterCallbackBase& QueryCallback, const ChaosInterface::FQueryDebugParams& DebugParams)
+void SweepHelper(const QueryGeomType& QueryGeom, const Chaos::ISpatialAcceleration<TPayload, Chaos::FReal, 3>& SpatialAcceleration, const FTransform& StartTM, const FVector& Dir, const float DeltaMagnitude, ChaosInterface::FSQHitBuffer<TSweepHit>& HitBuffer, EHitFlags OutputFlags, const ChaosInterface::FSceneQueryCommonParams& InCommonParams)
 {
 	using namespace Chaos;
 	using namespace ChaosInterface;
@@ -661,7 +704,7 @@ void SweepHelper(const QueryGeomType& QueryGeom, const Chaos::ISpatialAccelerati
 	const bool bIsValid = !Bounds.Min().ContainsNaN() && !Bounds.Max().ContainsNaN();
 	if (ensure(bIsValid))
 	{
-		TSQVisitor<QueryGeomType, TPayload, TSweepHit, std::is_same<TSweepHit, FSweepHit>::value> SweepVisitor(StartTM, Dir, HitBuffer, OutputFlags, QueryFilterData, QueryCallback, QueryGeom, Bounds, DebugParams);
+		TSQVisitor<QueryGeomType, TPayload, TSweepHit, std::is_same<TSweepHit, FSweepHit>::value> SweepVisitor(StartTM, Dir, HitBuffer, OutputFlags, QueryGeom, Bounds, InCommonParams);
 
 		HitBuffer.IncFlushCount();
 
@@ -676,7 +719,7 @@ void SweepHelper(const QueryGeomType& QueryGeom, const Chaos::ISpatialAccelerati
 
 
 template <typename QueryGeomType, typename TOverlapHit, typename TPayload>
-void OverlapHelper(const QueryGeomType& QueryGeom, const Chaos::ISpatialAcceleration<TPayload, Chaos::FReal, 3>& SpatialAcceleration, const FTransform& GeomPose, ChaosInterface::FSQHitBuffer<TOverlapHit>& HitBuffer, const ChaosInterface::FQueryFilterData& QueryFilterData, ICollisionQueryFilterCallbackBase& QueryCallback, const ChaosInterface::FQueryDebugParams& DebugParams)
+void OverlapHelper(const QueryGeomType& QueryGeom, const Chaos::ISpatialAcceleration<TPayload, Chaos::FReal, 3>& SpatialAcceleration, const FTransform& GeomPose, ChaosInterface::FSQHitBuffer<TOverlapHit>& HitBuffer, const ChaosInterface::FSceneQueryCommonParams& InCommonParams)
 {
 	using namespace Chaos;
 	using namespace ChaosInterface;
@@ -688,19 +731,37 @@ void OverlapHelper(const QueryGeomType& QueryGeom, const Chaos::ISpatialAccelera
 	{
 		HitBuffer.IncFlushCount();
 
-		bool bSkipNarrowPhase = QueryFilterData.flags & FPhysicsQueryFlag::eSKIPNARROWPHASE;
+		bool bSkipNarrowPhase = InCommonParams.Flags & FPhysicsQueryFlag::eSKIPNARROWPHASE;
 
 		constexpr bool bGTData = std::is_same<TOverlapHit, FOverlapHit>::value;
 		if (bSkipNarrowPhase)
 		{
-			TBPVisitor<QueryGeomType, TPayload, TOverlapHit, bGTData> OverlapVisitor(GeomPose, HitBuffer, QueryFilterData, QueryCallback, QueryGeom, DebugParams);
+			TBPVisitor<QueryGeomType, TPayload, TOverlapHit, bGTData> OverlapVisitor(GeomPose, HitBuffer, QueryGeom, InCommonParams);
 			SpatialAcceleration.Overlap(Bounds, OverlapVisitor);
 		}
 		else
 		{
-			TSQVisitor<QueryGeomType, TPayload, TOverlapHit, bGTData> OverlapVisitor(GeomPose, HitBuffer, QueryFilterData, QueryCallback, QueryGeom, Bounds, DebugParams);
+			TSQVisitor<QueryGeomType, TPayload, TOverlapHit, bGTData> OverlapVisitor(GeomPose, HitBuffer, QueryGeom, Bounds, InCommonParams);
 			SpatialAcceleration.Overlap(Bounds, OverlapVisitor);
 		}
 		HitBuffer.DecFlushCount();
 	}
 }
+
+PRAGMA_DISABLE_DEPRECATION_WARNINGS
+template <typename QueryGeomType, typename TSweepHit, typename TPayload>
+UE_DEPRECATED(5.8, "Use the overload that takes ChaosInterface::FSceneQueryCommonParams instead")
+void SweepHelper(const QueryGeomType& QueryGeom, const Chaos::ISpatialAcceleration<TPayload, Chaos::FReal, 3>& SpatialAcceleration, const FTransform& StartTM, const FVector& Dir, const float DeltaMagnitude, ChaosInterface::FSQHitBuffer<TSweepHit>& HitBuffer, EHitFlags OutputFlags, const ChaosInterface::FQueryFilterData& QueryFilterData, ICollisionQueryFilterCallbackBase& QueryCallback, const ChaosInterface::FQueryDebugParams& DebugParams)
+{
+	const ChaosInterface::FSceneQueryCommonParams CommonParams(QueryCallback, QueryFilterData, DebugParams);
+	SweepHelper(QueryGeom, SpatialAcceleration, StartTM, Dir, DeltaMagnitude, HitBuffer, OutputFlags, CommonParams);
+}
+
+template <typename QueryGeomType, typename TOverlapHit, typename TPayload>
+UE_DEPRECATED(5.8, "Use the overload that takes ChaosInterface::FSceneQueryCommonParams instead")
+void OverlapHelper(const QueryGeomType& QueryGeom, const Chaos::ISpatialAcceleration<TPayload, Chaos::FReal, 3>& SpatialAcceleration, const FTransform& GeomPose, ChaosInterface::FSQHitBuffer<TOverlapHit>& HitBuffer, const ChaosInterface::FQueryFilterData& QueryFilterData, ICollisionQueryFilterCallbackBase& QueryCallback, const ChaosInterface::FQueryDebugParams& DebugParams)
+{
+	const ChaosInterface::FSceneQueryCommonParams CommonParams(QueryCallback, QueryFilterData, DebugParams);
+	OverlapHelper(QueryGeom, SpatialAcceleration, GeomPose, HitBuffer, CommonParams);
+}
+PRAGMA_ENABLE_DEPRECATION_WARNINGS
diff --git a/Engine/Source/Runtime/PhysicsCore/Public/SceneQueryCommonParams.h b/Engine/Source/Runtime/PhysicsCore/Public/SceneQueryCommonParams.h
new file mode 100644
--- /dev/null
+++ b/Engine/Source/Runtime/PhysicsCore/Public/SceneQueryCommonParams.h
@@ -0,0 +1,35 @@
+// Copyright Epic Games, Inc. All Rights Reserved.
+
+#pragma once
+
+#include "PhysicsInterfaceTypesCore.h"
+#include "ChaosInterfaceWrapperCore.h"
+#include "Chaos/CollisionFilterData.h"
+
+class ICollisionQueryFilterCallback;
+
+namespace ChaosInterface
+{
+	struct FSceneQueryCommonParams
+	{
+		using FQueryFilterData = Chaos::Filter::FQueryFilterData;
+
+		PHYSICSCORE_API explicit FSceneQueryCommonParams(ICollisionQueryFilterCallback& InQueryCallback);
+		PHYSICSCORE_API explicit FSceneQueryCommonParams(ICollisionQueryFilterCallback& InQueryCallback, const FQueryFilterData& InFilterData, FChaosQueryFlags InFlags, const FQueryDebugParams& InDebugParams = {});
+		PHYSICSCORE_API explicit FSceneQueryCommonParams(ICollisionQueryFilterCallback& InQueryCallback, const FQueryFilterData& InFilterData, EQueryFlags InFlags, const FQueryDebugParams& InDebugParams = {});
+
+		PRAGMA_DISABLE_DEPRECATION_WARNINGS
+		// Temporary constructor for deprecated code.
+		UE_INTERNAL PHYSICSCORE_API explicit FSceneQueryCommonParams(ICollisionQueryFilterCallback& InQueryCallback, const FChaosQueryFilterData& InFilterData, const FQueryDebugParams& InDebugParams = {});
+		PRAGMA_ENABLE_DEPRECATION_WARNINGS
+
+		PHYSICSCORE_API ICollisionQueryFilterCallback& GetQueryCallback() const;
+
+		FQueryFilterData FilterData;
+		FChaosQueryFlags Flags;
+		FQueryDebugParams DebugParams;
+
+	private:
+		ICollisionQueryFilterCallback* QueryCallback = nullptr;
+	};
+} // namespace ChaosInterface
