Update FTriMeshCollisionData serialization to include the new bReducedMemoryRepresentation flag

#Synced-CL 48812894

diff --git a/Engine/Source/Runtime/Core/Public/UObject/FortniteMainBranchObjectVersions.inl b/Engine/Source/Runtime/Core/Public/UObject/FortniteMainBranchObjectVersions.inl
--- a/Engine/Source/Runtime/Core/Public/UObject/FortniteMainBranchObjectVersions.inl
+++ b/Engine/Source/Runtime/Core/Public/UObject/FortniteMainBranchObjectVersions.inl
@@ -715,6 +715,8 @@ FORTNITE_MAIN_VERSION(PCGMetadataOpStringTypeChange, 233)
 // Compress LiveLink SubSection data
 FORTNITE_MAIN_VERSION(LiveLinkCompressSubSectionData, 234)
 
+// Added option to quantize chaos tri mesh BVH's per-face bounds
+FORTNITE_MAIN_VERSION(OptionalQuantizePerFaceBoundsInChaosTriMeshBVH, 235)
 
 
 // -----<new versions can be added above this line>-------------------------------------------------
diff --git a/Engine/Source/Runtime/PhysicsCore/Private/PhysicsCore.cpp b/Engine/Source/Runtime/PhysicsCore/Private/PhysicsCore.cpp
--- a/Engine/Source/Runtime/PhysicsCore/Private/PhysicsCore.cpp
+++ b/Engine/Source/Runtime/PhysicsCore/Private/PhysicsCore.cpp
@@ -3,6 +3,8 @@
 #include "PhysicsCore.h"
 #include "Interface_CollisionDataProviderCore.h"
 
+#include "UObject/FortniteMainBranchObjectVersion.h"
+
 DEFINE_LOG_CATEGORY(LogPhysicsCore);
 
 #if WITH_EDITORONLY_DATA
@@ -13,6 +15,8 @@ void operator<<(class FArchive& Ar, FTriIndices& Indices)
 
 void operator<<(class FArchive& Ar, FTriMeshCollisionData& TriMeshCollisionData)
 {
+	Ar.UsingCustomVersion(FFortniteMainBranchObjectVersion::GUID);
+
 	Ar << TriMeshCollisionData.Vertices;
 	Ar << TriMeshCollisionData.Indices;
 	Ar << TriMeshCollisionData.MaterialIndices;
@@ -23,5 +27,10 @@ void operator<<(class FArchive& Ar, FTriMeshCollisionData& TriMeshCollisionData)
 	FArchive_Serialize_BitfieldBool(Ar, TriMeshCollisionData.bDeformableMesh);
 	FArchive_Serialize_BitfieldBool(Ar, TriMeshCollisionData.bFastCook);
 	FArchive_Serialize_BitfieldBool(Ar, TriMeshCollisionData.bDisableActiveEdgePrecompute);
+
+	if (Ar.CustomVer(FFortniteMainBranchObjectVersion::GUID) >= FFortniteMainBranchObjectVersion::OptionalQuantizePerFaceBoundsInChaosTriMeshBVH)
+	{
+		FArchive_Serialize_BitfieldBool(Ar, TriMeshCollisionData.bReducedMemoryRepresentation);
+	}
 }
 #endif
