PR #13551: Support FPhysicsObject stored in FChaosUserData and support PrimitiveComponents to not use BodyInstances but fallback to using FPhysicsObjects in when doing query result conversion

#Synced-CL 50101401

diff --git a/Engine/Source/Runtime/Engine/Private/Collision/CollisionConversions.cpp b/Engine/Source/Runtime/Engine/Private/Collision/CollisionConversions.cpp
--- a/Engine/Source/Runtime/Engine/Private/Collision/CollisionConversions.cpp
+++ b/Engine/Source/Runtime/Engine/Private/Collision/CollisionConversions.cpp
@@ -7,6 +7,7 @@
 #include "Engine/World.h"
 #include "Chaos/Capsule.h"
 #include "Chaos/ImplicitObjectType.h"
+#include "Chaos/PhysicsObjectInterface.h"
 #include "Components/PrimitiveComponent.h"
 
 #include "EngineLogs.h"
@@ -170,6 +171,12 @@ static void SetHitResultFromShapeAndFaceIndex(const FPhysicsShape& Shape,  const
 		OwningComponent = BodyInst->OwnerComponent;
 		OutResult.PhysicsObject = BodyInst->GetPhysicsActor() ? BodyInst->GetPhysicsActor()->GetPhysicsObject() : nullptr;
 	}
+	else if (Chaos::FPhysicsObject* PhysObject = GetPhysicsObjectFromUserData(Actor))
+	{
+		OutResult.PhysicsObject = PhysObject;
+		OutResult.Item = Chaos::FPhysicsObjectInterface::GetId(PhysObject);
+		OutResult.BoneName = Chaos::FPhysicsObjectInterface::GetName(PhysObject);
+	}
 	else
 	{
 		// Currently geom collections are registered with a primitive component user data, but maybe custom should be adapted
diff --git a/Engine/Source/Runtime/Engine/Private/Components/PrimitiveComponent.cpp b/Engine/Source/Runtime/Engine/Private/Components/PrimitiveComponent.cpp
--- a/Engine/Source/Runtime/Engine/Private/Components/PrimitiveComponent.cpp
+++ b/Engine/Source/Runtime/Engine/Private/Components/PrimitiveComponent.cpp
@@ -3649,7 +3649,7 @@ bool UPrimitiveComponent::LineTraceComponent(FHitResult& OutHit, const FVector S
 {
 	bool bHaveHit = false;
 
-	if (FBodyInstance* ThisBodyInstance = GetBodyInstance())
+	if (FBodyInstance* ThisBodyInstance = GetBodyInstance(); ThisBodyInstance && ThisBodyInstance->IsValidBodyInstance())
 	{
 		bHaveHit = ThisBodyInstance->LineTrace(OutHit, Start, End, Params.bTraceComplex, Params.bReturnPhysicalMaterial);
 	}
diff --git a/Engine/Source/Runtime/Engine/Private/PhysicsEngine/Experimental/ChaosInterfaceWrapper.cpp b/Engine/Source/Runtime/Engine/Private/PhysicsEngine/Experimental/ChaosInterfaceWrapper.cpp
--- a/Engine/Source/Runtime/Engine/Private/PhysicsEngine/Experimental/ChaosInterfaceWrapper.cpp
+++ b/Engine/Source/Runtime/Engine/Private/PhysicsEngine/Experimental/ChaosInterfaceWrapper.cpp
@@ -2,6 +2,7 @@
 
 #include "Physics/Experimental/ChaosInterfaceWrapper.h"
 #include "Chaos/ChaosUserEntity.h"
+#include "Chaos/PhysicsObject.h"
 #include "Physics/Experimental/PhysScene_Chaos.h"
 #include "PBDRigidsSolver.h"
 
@@ -54,6 +55,26 @@ namespace ChaosInterface
 		return nullptr;
 	}
 
+	Chaos::FPhysicsObject* GetPhysicsObjectFromUserData(const Chaos::FGeometryParticle& Actor)
+	{
+		void* UserData = Actor.UserData();
+		if (UserData)
+		{
+			Chaos::FPhysicsObject* PhysicsObject = FChaosUserData::Get<Chaos::FPhysicsObject>(UserData);
+			if (!PhysicsObject)
+			{
+				// Check if we appended a custom entity
+				FChaosUserEntityAppend* ChaosUserEntityAppend = FChaosUserData::Get<FChaosUserEntityAppend>(UserData);
+				if (ChaosUserEntityAppend)
+				{
+					PhysicsObject = FChaosUserData::Get<Chaos::FPhysicsObject>(ChaosUserEntityAppend->ChaosUserData);
+				}
+			}
+			return PhysicsObject;
+		}
+		return nullptr;
+	}
+
 	bool HasValidUserData(const Chaos::FGeometryParticle& Actor)
 	{
 		if (FChaosUserData* UserData = static_cast<FChaosUserData*>(Actor.UserData()))
diff --git a/Engine/Source/Runtime/Engine/Public/Physics/Experimental/ChaosInterfaceWrapper.h b/Engine/Source/Runtime/Engine/Public/Physics/Experimental/ChaosInterfaceWrapper.h
--- a/Engine/Source/Runtime/Engine/Public/Physics/Experimental/ChaosInterfaceWrapper.h
+++ b/Engine/Source/Runtime/Engine/Public/Physics/Experimental/ChaosInterfaceWrapper.h
@@ -20,6 +20,7 @@ using IExternalSpatialAcceleration = Chaos::ISpatialAcceleration<FExternalSpatia
 namespace Chaos
 {
 	class FPerShapeData;
+	struct FPhysicsObject;
 }
 
 namespace ChaosInterface
@@ -46,6 +47,7 @@ PRAGMA_ENABLE_DEPRECATION_WARNINGS
 FBodyInstance* GetUserData(const Chaos::FGeometryParticle& Actor);
 UPhysicalMaterial* GetUserData(const Chaos::FChaosPhysicsMaterial& Material);
 UPrimitiveComponent* GetPrimitiveComponentFromUserData(const Chaos::FGeometryParticle& Actor);
+Chaos::FPhysicsObject* GetPhysicsObjectFromUserData(const Chaos::FGeometryParticle& Actor);
 bool HasValidUserData(const Chaos::FGeometryParticle& Actor);
 
 }
diff --git a/Engine/Source/Runtime/PhysicsCore/Public/PhysicsInterfaceTypesCore.h b/Engine/Source/Runtime/PhysicsCore/Public/PhysicsInterfaceTypesCore.h
--- a/Engine/Source/Runtime/PhysicsCore/Public/PhysicsInterfaceTypesCore.h
+++ b/Engine/Source/Runtime/PhysicsCore/Public/PhysicsInterfaceTypesCore.h
@@ -191,8 +191,11 @@ struct FKShapeElem;
 /** Forward declarations */
 struct FKShapeElem;
 struct FCustomChaosPayload;
-struct FPhysicsObject;
 struct FChaosUserEntityAppend;
+namespace Chaos
+{
+	struct FPhysicsObject;
+}
 
 namespace EChaosUserDataType
 {
@@ -226,7 +229,7 @@ public:
 	FChaosUserData(FConstraintInstance* InPayload)		:Type(EChaosUserDataType::ConstraintInstance), Payload(InPayload) {}
 	FChaosUserData(UPrimitiveComponent* InPayload)		:Type(EChaosUserDataType::PrimitiveComponent), Payload(InPayload) {}
 	FChaosUserData(FKShapeElem* InPayload)				:Type(EChaosUserDataType::AggShape), Payload(InPayload) {}
-	FChaosUserData(FPhysicsObject* InPayload)			:Type(EChaosUserDataType::PhysicsObject), Payload(InPayload) {}
+	FChaosUserData(Chaos::FPhysicsObject* InPayload) 	:Type(EChaosUserDataType::PhysicsObject), Payload(InPayload) {}
 	FChaosUserData(FChaosUserEntityAppend* InPayload)	:Type(EChaosUserDataType::ChaosUserEntity), Payload(InPayload) {}
 	FChaosUserData(FCustomChaosPayload* InPayload)		:Type(EChaosUserDataType::CustomPayload), Payload(InPayload) {}
 	
@@ -249,7 +252,7 @@ template <> FORCEINLINE FPhysScene* FChaosUserData::Get(void* UserData)				{ if
 template <> FORCEINLINE FConstraintInstance* FChaosUserData::Get(void* UserData)	{ if (!UserData || ((FChaosUserData*)UserData)->Type != EChaosUserDataType::ConstraintInstance) { return nullptr; } return (FConstraintInstance*)((FChaosUserData*)UserData)->Payload; }
 template <> FORCEINLINE UPrimitiveComponent* FChaosUserData::Get(void* UserData)	{ if (!UserData || ((FChaosUserData*)UserData)->Type != EChaosUserDataType::PrimitiveComponent) { return nullptr; } return (UPrimitiveComponent*)((FChaosUserData*)UserData)->Payload; }
 template <> FORCEINLINE FKShapeElem* FChaosUserData::Get(void* UserData)	{ if (!UserData || ((FChaosUserData*)UserData)->Type != EChaosUserDataType::AggShape) { return nullptr; } return (FKShapeElem*)((FChaosUserData*)UserData)->Payload; }
-template <> FORCEINLINE FPhysicsObject* FChaosUserData::Get(void* UserData) { if (!UserData || ((FChaosUserData*)UserData)->Type != EChaosUserDataType::PhysicsObject) { return nullptr; } return (FPhysicsObject*)((FChaosUserData*)UserData)->Payload; }
+template <> FORCEINLINE Chaos::FPhysicsObject* FChaosUserData::Get(void* UserData) { if (!UserData || ((FChaosUserData*)UserData)->Type != EChaosUserDataType::PhysicsObject) { return nullptr; } return (Chaos::FPhysicsObject*)((FChaosUserData*)UserData)->Payload; }
 template <> FORCEINLINE FChaosUserEntityAppend* FChaosUserData::Get(void* UserData) { if (!UserData || ((FChaosUserData*)UserData)->Type != EChaosUserDataType::ChaosUserEntity) { return nullptr; } return (FChaosUserEntityAppend*)((FChaosUserData*)UserData)->Payload; }
 template <> FORCEINLINE FCustomChaosPayload* FChaosUserData::Get(void* UserData) { if (!UserData || ((FChaosUserData*)UserData)->Type != EChaosUserDataType::CustomPayload) { return nullptr; } return (FCustomChaosPayload*)((FChaosUserData*)UserData)->Payload; }
 
@@ -259,7 +262,7 @@ template <> FORCEINLINE void FChaosUserData::Set(void* UserData, FPhysScene* Pay
 template <> FORCEINLINE void FChaosUserData::Set(void* UserData, FConstraintInstance* Payload)		{ check(UserData); ((FChaosUserData*)UserData)->Type = EChaosUserDataType::ConstraintInstance; ((FChaosUserData*)UserData)->Payload = Payload; }
 template <> FORCEINLINE void FChaosUserData::Set(void* UserData, UPrimitiveComponent* Payload)		{ check(UserData); ((FChaosUserData*)UserData)->Type = EChaosUserDataType::PrimitiveComponent; ((FChaosUserData*)UserData)->Payload = Payload; }
 template <> FORCEINLINE void FChaosUserData::Set(void* UserData, FKShapeElem* Payload)	{ check(UserData); ((FChaosUserData*)UserData)->Type = EChaosUserDataType::AggShape; ((FChaosUserData*)UserData)->Payload = Payload; }
-template <> FORCEINLINE void FChaosUserData::Set(void* UserData, FPhysicsObject* Payload) { check(UserData); ((FChaosUserData*)UserData)->Type = EChaosUserDataType::PhysicsObject; ((FChaosUserData*)UserData)->Payload = Payload; }
+template <> FORCEINLINE void FChaosUserData::Set(void* UserData, Chaos::FPhysicsObject* Payload) { check(UserData); ((FChaosUserData*)UserData)->Type = EChaosUserDataType::PhysicsObject; ((FChaosUserData*)UserData)->Payload = Payload; }
 template <> FORCEINLINE void FChaosUserData::Set(void* UserData, FChaosUserEntityAppend* Payload) { check(UserData); ((FChaosUserData*)UserData)->Type = EChaosUserDataType::ChaosUserEntity; ((FChaosUserData*)UserData)->Payload = Payload; }
 template <> FORCEINLINE void FChaosUserData::Set(void* UserData, FCustomChaosPayload* Payload) { check(UserData); ((FChaosUserData*)UserData)->Type = EChaosUserDataType::CustomPayload; ((FChaosUserData*)UserData)->Payload = Payload; }
 
