[Chaos] Fixed bodies that start their simulation asleep firing a spurious wake event while remaining asleep. This was due to the iniitialisation of the velocity for the body triggering the wake even incorrectly.

#Synced-CL 50653012

diff --git a/Engine/Source/Runtime/Engine/Private/PhysicsEngine/BodyInstance.cpp b/Engine/Source/Runtime/Engine/Private/PhysicsEngine/BodyInstance.cpp
--- a/Engine/Source/Runtime/Engine/Private/PhysicsEngine/BodyInstance.cpp
+++ b/Engine/Source/Runtime/Engine/Private/PhysicsEngine/BodyInstance.cpp
@@ -4676,6 +4676,10 @@ void FBodyInstance::InitDynamicProperties_AssumesLocked()
 			{
 				FPhysicsInterface::SetWakeCounter_AssumesLocked(GetPhysicsActor(), 0.0f);
 				FPhysicsInterface::PutToSleep_AssumesLocked(GetPhysicsActor());
+
+				// Initialising the velocity above will trigger a wake event on a sleeping actor so although the
+				// actor will stay sleeping the user will receive a wake event unless we clear it prior to simulation
+				FPhysicsInterface::ClearWakeEvent_AssumesLocked(GetPhysicsActor());
 			}
 		}
 
diff --git a/Engine/Source/Runtime/PhysicsCore/Private/ChaosEngineInterface.cpp b/Engine/Source/Runtime/PhysicsCore/Private/ChaosEngineInterface.cpp
--- a/Engine/Source/Runtime/PhysicsCore/Private/ChaosEngineInterface.cpp
+++ b/Engine/Source/Runtime/PhysicsCore/Private/ChaosEngineInterface.cpp
@@ -460,6 +460,11 @@ void FChaosEngineInterface::SetSendsSleepNotifies_AssumesLocked(const FPhysicsAc
 	//check(bSendSleepNotifies == false);
 }
 
+void FChaosEngineInterface::ClearWakeEvent_AssumesLocked(const FPhysicsActorHandle& InActorReference)
+{
+	InActorReference->GetGameThreadAPI().ClearEvents();
+}
+
 void FChaosEngineInterface::PutToSleep_AssumesLocked(const FPhysicsActorHandle& InActorReference)
 {
 	// NOTE: We want to set the state whether or not it's asleep - if we currently think we're
diff --git a/Engine/Source/Runtime/PhysicsCore/Public/Chaos/ChaosEngineInterface.h b/Engine/Source/Runtime/PhysicsCore/Public/Chaos/ChaosEngineInterface.h
--- a/Engine/Source/Runtime/PhysicsCore/Public/Chaos/ChaosEngineInterface.h
+++ b/Engine/Source/Runtime/PhysicsCore/Public/Chaos/ChaosEngineInterface.h
@@ -408,6 +408,7 @@ public:
 	static PHYSICSCORE_API float GetMass_AssumesLocked(const FPhysicsActorHandle& InActorReference);
 
 	static PHYSICSCORE_API void SetSendsSleepNotifies_AssumesLocked(const FPhysicsActorHandle& InActorReference,bool bSendSleepNotifies);
+	static PHYSICSCORE_API void ClearWakeEvent_AssumesLocked(const FPhysicsActorHandle& InActorReference);
 	static PHYSICSCORE_API void PutToSleep_AssumesLocked(const FPhysicsActorHandle& InActorReference);
 	static PHYSICSCORE_API void WakeUp_AssumesLocked(const FPhysicsActorHandle& InActorReference);
 
