[Chaos]Fix to WaitPhysScenes to wait for async tasks not tallied in CompletionEvents. This fixes teardown issues in async mode potantially preventing garbage collection in cases where async code is holding onto UObjects.

#Synced-CL 52152484

diff --git a/Engine/Source/Runtime/PhysicsCore/Private/ChaosScene.cpp b/Engine/Source/Runtime/PhysicsCore/Private/ChaosScene.cpp
--- a/Engine/Source/Runtime/PhysicsCore/Private/ChaosScene.cpp
+++ b/Engine/Source/Runtime/PhysicsCore/Private/ChaosScene.cpp
@@ -583,6 +583,18 @@ void FChaosScene::WaitPhysScenes()
 		QUICK_SCOPE_CYCLE_COUNTER(STAT_FPhysScene_WaitPhysScenes);
 		FTaskGraphInterface::Get().WaitUntilTasksComplete(CompletionEvents,ENamedThreads::GameThread);
 	}
+
+	// CompletionEvents only captures the return value of AdvanceAndDispatch_External, which in
+	// async physics mode (IsUsingAsyncResults()==true) intentionally returns the event from the
+	// previous frame rather than the newly-dispatched tasks. Those tasks run independently and
+	// are not tracked in CompletionEvents. WaitOnPendingTasks_External drains them directly.
+	for (Chaos::FPhysicsSolverBase* Solver : GetPhysicsSolvers())
+	{
+		if (Solver && Solver->IsUsingAsyncResults())
+		{
+			Solver->WaitOnPendingTasks_External();
+		}
+	}
 }
 
 FGraphEventArray FChaosScene::GetCompletionEvents()
